Merging //UE4/Main@9388041 to Dev-RenderPlat-Staging

diff --git a/Engine/Source/Editor/ClothPainter/Private/SClothAssetSelector.cpp b/Engine/Source/Editor/ClothPainter/Private/SClothAssetSelector.cpp
--- a/Engine/Source/Editor/ClothPainter/Private/SClothAssetSelector.cpp
+++ b/Engine/Source/Editor/ClothPainter/Private/SClothAssetSelector.cpp
@@ -7,7 +7,8 @@
 #include "DetailLayoutBuilder.h"
 #include "Widgets/Input/SButton.h"
 #include "Widgets/Images/SImage.h"
-#include "Assets/ClothingAsset.h"
+#include "ClothingAsset.h"
+#include "ClothPhysicalMeshData.h"
 #include "Framework/MultiBox/MultiBoxBuilder.h"
 #include "Widgets/Input/SCheckBox.h"
 #include "Widgets/Text/SInlineEditableTextBlock.h"
@@ -32,16 +33,16 @@
 
 #define LOCTEXT_NAMESPACE "ClothAssetSelector"
 
-FClothParameterMask_PhysMesh* FClothingMaskListItem::GetMask()
+FPointWeightMap* FClothingMaskListItem::GetMask()
 {
-	if(UClothingAsset* Asset = ClothingAsset.Get())
+	if(UClothingAssetCommon* Asset = ClothingAsset.Get())
 	{
 		if(Asset->IsValidLod(LodIndex))
 		{
-			FClothLODData& LodData = Asset->LodData[LodIndex];
-			if(LodData.ParameterMasks.IsValidIndex(MaskIndex))
+			UClothLODDataBase* LodData = Asset->ClothLodData[LodIndex];
+			if(LodData->ParameterMasks.IsValidIndex(MaskIndex))
 			{
-				return &LodData.ParameterMasks[MaskIndex];
+				return &LodData->ParameterMasks[MaskIndex];
 			}
 		}
 	}
@@ -49,14 +50,16 @@ FClothParameterMask_PhysMesh* FClothingMaskListItem::GetMask()
 	return nullptr;
 }
 
-USkeletalMesh* FClothingMaskListItem::GetOwningMesh()
+UClothPhysicalMeshDataBase* FClothingMaskListItem::GetMeshData()
 {
-	if(UClothingAsset* Asset = ClothingAsset.Get())
-	{
-		return Cast<USkeletalMesh>(Asset->GetOuter());
-	}
+	UClothingAssetCommon * Asset = ClothingAsset.Get();
+	return Asset && Asset->IsValidLod(LodIndex) ? Asset->ClothLodData[LodIndex]->PhysicalMeshData : nullptr;
+}
 
-	return nullptr;
+USkeletalMesh* FClothingMaskListItem::GetOwningMesh()
+{
+	UClothingAssetCommon * Asset = ClothingAsset.Get();
+	return Asset ? Cast<USkeletalMesh>(Asset->GetOuter()) : nullptr;
 }
 
 class SAssetListRow : public STableRow<TSharedPtr<FClothingAssetListItem>>
@@ -106,7 +109,7 @@ public:
 	{
 		if(Item.IsValid())
 		{
-			if(UClothingAsset* Asset = Item->ClothingAsset.Get())
+			if(UClothingAssetCommon* Asset = Item->ClothingAsset.Get())
 			{
 				FText TrimText = FText::TrimPrecedingAndTrailing(InText);
 
@@ -115,9 +118,9 @@ public:
 					FName NewName(*TrimText.ToString());
 
 					// Check for an existing object, and if we find one build a unique name based on the request
-					if(UObject* ExistingObject = StaticFindObject(UClothingAsset::StaticClass(), Asset->GetOuter(), *NewName.ToString()))
+					if(UObject* ExistingObject = StaticFindObject(UClothingAssetCommon::StaticClass(), Asset->GetOuter(), *NewName.ToString()))
 					{
-						NewName = MakeUniqueObjectName(Asset->GetOuter(), UClothingAsset::StaticClass(), FName(*TrimText.ToString()));
+						NewName = MakeUniqueObjectName(Asset->GetOuter(), UClothingAssetCommon::StaticClass(), FName(*TrimText.ToString()));
 					}
 
 					Asset->Rename(*NewName.ToString(), Asset->GetOuter());
@@ -181,7 +184,7 @@ private:
 
 	void DeleteAsset()
 	{
-		if(UClothingAsset* Asset = Item->ClothingAsset.Get())
+		if(UClothingAssetCommon* Asset = Item->ClothingAsset.Get())
 		{
 			if(USkeletalMesh* SkelMesh = Cast<USkeletalMesh>(Asset->GetOuter()))
 			{
@@ -229,7 +232,7 @@ private:
 
 	void ReimportAsset()
 	{
-		if(UClothingAsset* Asset = Item->ClothingAsset.Get())
+		if(UClothingAssetCommon* Asset = Item->ClothingAsset.Get())
 		{
 			if(USkeletalMesh* SkelMesh = Cast<USkeletalMesh>(Asset->GetOuter()))
 			{
@@ -289,32 +292,32 @@ private:
 			return;
 		}
 
-		if(UClothingAsset* Asset = Item->ClothingAsset.Get())
+		if(UClothingAssetCommon* Asset = Item->ClothingAsset.Get())
 		{
 			const int32 NumLods = Asset->GetNumLods();
 
 			for(int32 CurrIndex = 0; CurrIndex < NumLods - 1; ++CurrIndex)
 			{
-				FClothLODData& SourceLod = Asset->LodData[CurrIndex];
-				FClothLODData& DestLod = Asset->LodData[CurrIndex + 1];
+				UClothLODDataBase* SourceLod = Asset->ClothLodData[CurrIndex];
+				UClothLODDataBase* DestLod = Asset->ClothLodData[CurrIndex + 1];
 
-				DestLod.ParameterMasks.Reset();
+				DestLod->ParameterMasks.Reset();
 
-				for(FClothParameterMask_PhysMesh& SourceMask : SourceLod.ParameterMasks)
+				for(FPointWeightMap& SourceMask : SourceLod->ParameterMasks)
 				{
-					DestLod.ParameterMasks.AddDefaulted();
-					FClothParameterMask_PhysMesh& DestMask = DestLod.ParameterMasks.Last();
+					DestLod->ParameterMasks.AddDefaulted();
+					FPointWeightMap& DestMask = DestLod->ParameterMasks.Last();
 
-					DestMask.MaskName = SourceMask.MaskName;
+					DestMask.Name = SourceMask.Name;
 					DestMask.bEnabled = SourceMask.bEnabled;
 					DestMask.CurrentTarget = SourceMask.CurrentTarget;
 
 					ClothingMeshUtils::FVertexParameterMapper ParameterMapper(
-						DestLod.PhysicalMeshData.Vertices,
-						DestLod.PhysicalMeshData.Normals,
-						SourceLod.PhysicalMeshData.Vertices,
-						SourceLod.PhysicalMeshData.Normals,
-						SourceLod.PhysicalMeshData.Indices
+						DestLod->PhysicalMeshData->Vertices,
+						DestLod->PhysicalMeshData->Normals,
+						SourceLod->PhysicalMeshData->Vertices,
+						SourceLod->PhysicalMeshData->Normals,
+						SourceLod->PhysicalMeshData->Indices
 					);
 
 					ParameterMapper.Map(SourceMask.GetValueArray(), DestMask.Values);
@@ -330,7 +333,7 @@ private:
 			return false;
 		}
 
-		if(UClothingAsset* Asset = Item->ClothingAsset.Get())
+		if(UClothingAssetCommon* Asset = Item->ClothingAsset.Get())
 		{
 			if(Asset->GetNumLods() > 1)
 			{
@@ -395,8 +398,8 @@ public:
 
 		if(InColumnName == Column_CurrentTarget)
 		{
-			FClothParameterMask_PhysMesh* Mask = Item->GetMask();
-			UEnum* Enum = StaticEnum<MaskTarget_PhysMesh>();
+			FPointWeightMap* Mask = Item->GetMask();
+			UEnum* Enum = Item->GetMeshData()->GetFloatArrayTargets();
 			if(Enum && Mask)
 			{
 				return SNew(STextBlock).Text(Enum->GetDisplayNameTextByIndex((int32)Mask->CurrentTarget));
@@ -410,9 +413,9 @@ public:
 	{
 		if(Item.IsValid())
 		{
-			if(FClothParameterMask_PhysMesh* Mask = Item->GetMask())
+			if(FPointWeightMap* Mask = Item->GetMask())
 			{
-				return FText::FromName(Mask->MaskName);
+				return FText::FromName(Mask->Name);
 			}
 		}
 
@@ -423,10 +426,10 @@ public:
 	{
 		if(Item.IsValid())
 		{
-			if(FClothParameterMask_PhysMesh* Mask = Item->GetMask())
+			if(FPointWeightMap* Mask = Item->GetMask())
 			{
 				FText TrimText = FText::TrimPrecedingAndTrailing(InText);
-				Mask->MaskName = FName(*TrimText.ToString());
+				Mask->Name = FName(*TrimText.ToString());
 			}
 		}
 	}
@@ -436,7 +439,7 @@ public:
 		// Spawn menu
 		if(MouseEvent.GetEffectingButton() == EKeys::RightMouseButton && Item.IsValid())
 		{
-			if(FClothParameterMask_PhysMesh* Mask = Item->GetMask())
+			if(FPointWeightMap* Mask = Item->GetMask())
 			{
 				FMenuBuilder Builder(true, UICommandList);
 
@@ -483,17 +486,16 @@ private:
 		);
 	}
 
-	FClothLODData* GetCurrentLod() const
+	UClothLODDataBase* GetCurrentLod() const
 	{
 		if(Item.IsValid())
 		{
-			if(UClothingAsset* Asset = Item->ClothingAsset.Get())
+			if(UClothingAssetCommon* Asset = Item->ClothingAsset.Get())
 			{
-				if(Asset->LodData.IsValidIndex(Item->LodIndex))
+				if(Asset->ClothLodData.IsValidIndex(Item->LodIndex))
 				{
-					FClothLODData& LodData = Asset->LodData[Item->LodIndex];
-
-					return &LodData;
+					UClothLODDataBase* LodData = Asset->ClothLodData[Item->LodIndex];
+					return LodData;
 				}
 			}
 		}
@@ -510,7 +512,7 @@ private:
 			FScopedTransaction CurrTransaction(LOCTEXT("DeleteMask_Transaction", "Delete clothing parameter mask."));
 			Item->ClothingAsset->Modify();
 
-		if(FClothLODData* LodData = GetCurrentLod())
+		if(UClothLODDataBase* LodData = GetCurrentLod())
 		{
 			if(LodData->ParameterMasks.IsValidIndex(Item->MaskIndex))
 			{
@@ -519,7 +521,7 @@ private:
 				// We've removed a mask, so it will need to be applied to the clothing data
 				if(Item.IsValid())
 				{
-					if(UClothingAsset* Asset = Item->ClothingAsset.Get())
+					if(UClothingAssetCommon* Asset = Item->ClothingAsset.Get())
 					{
 						Asset->ApplyParameterMasks();
 					}
@@ -540,11 +542,10 @@ private:
 			FScopedTransaction CurrTransaction(LOCTEXT("SetMaskTarget_Transaction", "Set clothing parameter mask target."));
 			Item->ClothingAsset->Modify();
 
-			if(FClothParameterMask_PhysMesh* Mask = Item->GetMask())
+			if(FPointWeightMap* Mask = Item->GetMask())
 			{
-				Mask->CurrentTarget = (MaskTarget_PhysMesh)InTargetEntryIndex;
-
-				if(Mask->CurrentTarget == MaskTarget_PhysMesh::None)
+				Mask->CurrentTarget = (uint8)InTargetEntryIndex;
+				if(Mask->CurrentTarget == 0)//(uint8)MaskTarget_PhysMesh::None)
 				{
 					// Make sure to disable this mask if it has no valid target
 					Mask->bEnabled = false;
@@ -559,7 +560,8 @@ private:
 	{
 		Builder.BeginSection(NAME_None, LOCTEXT("MaskTargets_SectionName", "Targets"));
 		{
-			UEnum* Enum = StaticEnum<MaskTarget_PhysMesh>();
+			//UEnum* Enum = StaticEnum<MaskTarget_PhysMesh>();
+			UEnum* Enum = Item->GetMeshData()->GetFloatArrayTargets();
 			if(Enum)
 			{
 				const int32 NumEntries = Enum->NumEnums();
@@ -583,9 +585,9 @@ private:
 	{
 		if (AssetSelectorPtr.IsValid())
 		{
-			UClothingAsset* ClothingAsset = AssetSelectorPtr.Pin()->GetSelectedAsset().Get();
+			UClothingAssetCommon* ClothingAsset = AssetSelectorPtr.Pin()->GetSelectedAsset().Get();
 			int32 LOD = AssetSelectorPtr.Pin()->GetSelectedLod();
-			FClothParameterMask_PhysMesh* Mask = Item->GetMask();
+			FPointWeightMap* Mask = Item->GetMask();
 
 			TSharedRef<SWidget> Widget = SNew(SCopyVertexColorSettingsPanel, ClothingAsset, LOD, Mask);
 
@@ -598,7 +600,7 @@ private:
 	{
 		if(InItem.IsValid())
 		{
-			if(FClothParameterMask_PhysMesh* Mask = InItem->GetMask())
+			if(FPointWeightMap* Mask = InItem->GetMask())
 			{
 				return Mask->bEnabled ? ECheckBoxState::Checked : ECheckBoxState::Unchecked;
 			}
@@ -611,9 +613,9 @@ private:
 	{
 		if(InItem.IsValid())
 		{
-			if(FClothParameterMask_PhysMesh* Mask = InItem->GetMask())
+			if(FPointWeightMap* Mask = InItem->GetMask())
 			{
-				return Mask->CurrentTarget != MaskTarget_PhysMesh::None;
+				return Mask->CurrentTarget != 0; // (uint8)MaskTarget_PhysMesh::None;
 			}
 		}
 
@@ -624,7 +626,7 @@ private:
 	{
 		if(InItem.IsValid())
 		{
-			if(FClothParameterMask_PhysMesh* Mask = InItem->GetMask())
+			if(FPointWeightMap* Mask = InItem->GetMask())
 			{
 				bool bNewEnableState = InState == ECheckBoxState::Checked;
 
@@ -633,16 +635,16 @@ private:
 					if(bNewEnableState)
 					{
 						// Disable all other masks that affect this target
-						if(UClothingAsset* Asset = InItem->ClothingAsset.Get())
+						if(UClothingAssetCommon* Asset = InItem->ClothingAsset.Get())
 						{
-							if(Asset->LodData.IsValidIndex(InItem->LodIndex))
+							if(Asset->ClothLodData.IsValidIndex(InItem->LodIndex))
 							{
-								FClothLODData& LodData = Asset->LodData[InItem->LodIndex];
+								UClothLODDataBase* LodData = Asset->ClothLodData[InItem->LodIndex];
 
-								TArray<FClothParameterMask_PhysMesh*> AllTargetMasks;
-								LodData.GetParameterMasksForTarget(Mask->CurrentTarget, AllTargetMasks);
+								TArray<FPointWeightMap*> AllTargetMasks;
+								LodData->GetParameterMasksForTarget(Mask->CurrentTarget, AllTargetMasks);
 
-								for(FClothParameterMask_PhysMesh* TargetMask : AllTargetMasks)
+								for(FPointWeightMap* TargetMask : AllTargetMasks)
 								{
 									if(TargetMask && TargetMask != Mask)
 									{
@@ -668,7 +670,7 @@ private:
 };
 
 FName SMaskListRow::Column_Enabled(TEXT("Enabled"));
-FName SMaskListRow::Column_MaskName(TEXT("MaskName"));
+FName SMaskListRow::Column_MaskName(TEXT("Name"));
 FName SMaskListRow::Column_CurrentTarget(TEXT("CurrentTarget"));
 
 SClothAssetSelector::~SClothAssetSelector()
@@ -929,7 +931,7 @@ void SClothAssetSelector::Construct(const FArguments& InArgs, USkeletalMesh* InM
 	RefreshMaskList();
 }
 
-TWeakObjectPtr<UClothingAsset> SClothAssetSelector::GetSelectedAsset() const
+TWeakObjectPtr<UClothingAssetCommon> SClothAssetSelector::GetSelectedAsset() const
 {
 	return SelectedAsset;
 	
@@ -976,7 +978,7 @@ void SClothAssetSelector::OnCopyClothingAssetSelected(const FAssetData& AssetDat
 
 		for (UClothingAssetBase* ClothingAsset : SourceSkelMesh->MeshClothingAssets)
 		{
-			UClothingAsset* NewAsset = Cast<UClothingAsset>(AssetFactory->CreateFromExistingCloth(Mesh, SourceSkelMesh, ClothingAsset));
+			UClothingAssetCommon* NewAsset = Cast<UClothingAssetCommon>(AssetFactory->CreateFromExistingCloth(Mesh, SourceSkelMesh, ClothingAsset));
 			Mesh->AddClothingAsset(NewAsset);
 		}
 		OnRefresh();
@@ -1025,9 +1027,9 @@ TSharedRef<SWidget> SClothAssetSelector::OnGetLodMenu()
 
 	int32 NumLods = 0;
 
-	if(UClothingAsset* CurrAsset = SelectedAsset.Get())
+	if(UClothingAssetCommon* CurrAsset = SelectedAsset.Get())
 	{
-		NumLods = CurrAsset->LodData.Num();
+		NumLods = CurrAsset->GetNumLods();
 	}
 
 	if(NumLods == 0)
@@ -1063,7 +1065,7 @@ FText SClothAssetSelector::GetLodButtonText() const
 
 TSharedRef<ITableRow> SClothAssetSelector::OnGenerateWidgetForClothingAssetItem(TSharedPtr<FClothingAssetListItem> InItem, const TSharedRef<STableViewBase>& OwnerTable)
 {
-	if(UClothingAsset* Asset = InItem->ClothingAsset.Get())
+	if(UClothingAssetCommon* Asset = InItem->ClothingAsset.Get())
 	{
 		return SNew(SAssetListRow, OwnerTable, InItem)
 			.OnInvalidateList(this, &SClothAssetSelector::OnRefresh);
@@ -1086,7 +1088,7 @@ void SClothAssetSelector::OnAssetListSelectionChanged(TSharedPtr<FClothingAssetL
 
 TSharedRef<ITableRow> SClothAssetSelector::OnGenerateWidgetForMaskItem(TSharedPtr<FClothingMaskListItem> InItem, const TSharedRef<STableViewBase>& OwnerTable)
 {
-	if(FClothParameterMask_PhysMesh* Mask = InItem->GetMask())
+	if(FPointWeightMap* Mask = InItem->GetMask())
 	{
 		return SNew(SMaskListRow, OwnerTable, InItem, SharedThis(this))
 			.OnInvalidateList(this, &SClothAssetSelector::OnRefresh);
@@ -1113,19 +1115,19 @@ void SClothAssetSelector::OnMaskSelectionChanged(TSharedPtr<FClothingMaskListIte
 
 FReply SClothAssetSelector::AddNewMask()
 {
-	if(UClothingAsset* Asset = SelectedAsset.Get())
+	if(UClothingAssetCommon* Asset = SelectedAsset.Get())
 	{
-		if(Asset->LodData.IsValidIndex(SelectedLod))
+		if(Asset->ClothLodData.IsValidIndex(SelectedLod))
 		{
-			FClothLODData& LodData = Asset->LodData[SelectedLod];
-			const int32 NumRequiredValues = LodData.PhysicalMeshData.Vertices.Num();
+			UClothLODDataBase* LodData = Asset->ClothLodData[SelectedLod];
+			const int32 NumRequiredValues = LodData->PhysicalMeshData->Vertices.Num();
 
-			LodData.ParameterMasks.AddDefaulted();
+			LodData->ParameterMasks.AddDefaulted();
 
-			FClothParameterMask_PhysMesh& NewMask = LodData.ParameterMasks.Last();
+			FPointWeightMap& NewMask = LodData->ParameterMasks.Last();
 
-			NewMask.MaskName = TEXT("New Mask");
-			NewMask.CurrentTarget = MaskTarget_PhysMesh::None;
+			NewMask.Name = TEXT("New Mask");
+			NewMask.CurrentTarget = 0;// (uint8)MaskTarget_PhysMesh::None;
 			NewMask.Values.AddZeroed(NumRequiredValues);
 
 			OnRefresh();
@@ -1148,7 +1150,7 @@ void SClothAssetSelector::OnRefresh()
 
 void SClothAssetSelector::RefreshAssetList()
 {
-	UClothingAsset* CurrSelectedAsset = nullptr;
+	UClothingAssetCommon* CurrSelectedAsset = nullptr;
 	int32 SelectedItem = INDEX_NONE;
 
 	if(AssetList.IsValid())
@@ -1166,7 +1168,7 @@ void SClothAssetSelector::RefreshAssetList()
 
 	for(UClothingAssetBase* Asset : Mesh->MeshClothingAssets)
 	{
-		UClothingAsset* ConcreteAsset = Cast<UClothingAsset>(Asset);
+		UClothingAssetCommon* ConcreteAsset = Cast<UClothingAssetCommon>(Asset);
 
 		TSharedPtr<FClothingAssetListItem> Entry = MakeShareable(new FClothingAssetListItem);
 
@@ -1218,11 +1220,11 @@ void SClothAssetSelector::RefreshMaskList()
 
 	MaskListItems.Empty();
 
-	UClothingAsset* Asset = SelectedAsset.Get();
+	UClothingAssetCommon* Asset = SelectedAsset.Get();
 	if(Asset && Asset->IsValidLod(SelectedLod))
 	{
-		FClothLODData& LodData = Asset->LodData[SelectedLod];
-		const int32 NumMasks = LodData.ParameterMasks.Num();
+		UClothLODDataBase* LodData = Asset->ClothLodData[SelectedLod];
+		const int32 NumMasks = LodData->ParameterMasks.Num();
 
 		for(int32 Index = 0; Index < NumMasks; ++Index)
 		{
@@ -1271,11 +1273,11 @@ void SClothAssetSelector::OnClothingLodSelected(int32 InNewLod)
 		SetSelectedLod(InNewLod);
 
 		int32 NewMaskSelection = INDEX_NONE;
-		if(SelectedAsset->LodData.IsValidIndex(SelectedLod))
+		if(SelectedAsset->ClothLodData.IsValidIndex(SelectedLod))
 		{
-			FClothLODData& LodData = SelectedAsset->LodData[SelectedLod];
+			UClothLODDataBase* LodData = SelectedAsset->ClothLodData[SelectedLod];
 
-			if(LodData.ParameterMasks.Num() > 0)
+			if(LodData->ParameterMasks.Num() > 0)
 			{
 				NewMaskSelection = 0;
 			}
@@ -1285,20 +1287,20 @@ void SClothAssetSelector::OnClothingLodSelected(int32 InNewLod)
 	}
 }
 
-void SClothAssetSelector::SetSelectedAsset(TWeakObjectPtr<UClothingAsset> InSelectedAsset)
+void SClothAssetSelector::SetSelectedAsset(TWeakObjectPtr<UClothingAssetCommon> InSelectedAsset)
 {
 	SelectedAsset = InSelectedAsset;
 
 	RefreshMaskList();
 
-	if(UClothingAsset* NewAsset = SelectedAsset.Get())
+	if(UClothingAssetCommon* NewAsset = SelectedAsset.Get())
 	{
-		if(NewAsset->LodData.Num() > 0)
+		if(NewAsset->GetNumLods() > 0)
 		{
 			SetSelectedLod(0);
 
-			FClothLODData& LodData = NewAsset->LodData[SelectedLod];
-			if(LodData.ParameterMasks.Num() > 0)
+			UClothLODDataBase* LodData = NewAsset->ClothLodData[SelectedLod];
+			if(LodData->ParameterMasks.Num() > 0)
 			{
 				SetSelectedMask(0);
 			}
diff --git a/Engine/Source/Editor/ClothPainter/Private/SClothAssetSelector.h b/Engine/Source/Editor/ClothPainter/Private/SClothAssetSelector.h
--- a/Engine/Source/Editor/ClothPainter/Private/SClothAssetSelector.h
+++ b/Engine/Source/Editor/ClothPainter/Private/SClothAssetSelector.h
@@ -10,12 +10,13 @@
 
 struct FAssetData;
 class USkeletalMesh;
-class UClothingAsset;
-struct FClothParameterMask_PhysMesh;
+class UClothingAssetCommon;
+class UClothPhysicalMeshDataBase;
+struct FPointWeightMap;
 
 struct FClothingAssetListItem
 {
-	TWeakObjectPtr<UClothingAsset> ClothingAsset;
+	TWeakObjectPtr<UClothingAssetCommon> ClothingAsset;
 };
 
 struct FClothingMaskListItem
@@ -25,10 +26,11 @@ struct FClothingMaskListItem
 		, MaskIndex(INDEX_NONE)
 	{}
 
-	FClothParameterMask_PhysMesh* GetMask();
+	FPointWeightMap* GetMask();
+	UClothPhysicalMeshDataBase* GetMeshData();
 	USkeletalMesh* GetOwningMesh();
 
-	TWeakObjectPtr<UClothingAsset> ClothingAsset;
+	TWeakObjectPtr<UClothingAssetCommon> ClothingAsset;
 	int32 LodIndex;
 	int32 MaskIndex;
 };
@@ -36,7 +38,7 @@ struct FClothingMaskListItem
 typedef SListView<TSharedPtr<FClothingAssetListItem>> SAssetList;
 typedef SListView<TSharedPtr<FClothingMaskListItem>> SMaskList;
 
-DECLARE_DELEGATE_ThreeParams(FOnClothAssetSelectionChanged, TWeakObjectPtr<UClothingAsset>, int32, int32);
+DECLARE_DELEGATE_ThreeParams(FOnClothAssetSelectionChanged, TWeakObjectPtr<UClothingAssetCommon>, int32, int32);
 
 class SClothAssetSelector : public SCompoundWidget, public FEditorUndoClient
 {
@@ -50,7 +52,7 @@ public:
 
 	void Construct(const FArguments& InArgs, USkeletalMesh* InMesh);
 
-	TWeakObjectPtr<UClothingAsset> GetSelectedAsset() const;
+	TWeakObjectPtr<UClothingAssetCommon> GetSelectedAsset() const;
 	int32 GetSelectedLod() const;
 	int32 GetSelectedMask() const;
 
@@ -92,7 +94,7 @@ protected:
 	void OnClothingLodSelected(int32 InNewLod);
 
 	// Setters for the list selections so we can handle list selections changing properly
-	void SetSelectedAsset(TWeakObjectPtr<UClothingAsset> InSelectedAsset);
+	void SetSelectedAsset(TWeakObjectPtr<UClothingAssetCommon> InSelectedAsset);
 	void SetSelectedLod(int32 InLodIndex, bool bRefreshMasks = true);
 	void SetSelectedMask(int32 InMaskIndex);
 
@@ -110,7 +112,7 @@ protected:
 	TArray<TSharedPtr<FClothingMaskListItem>> MaskListItems;
 
 	// Currently selected clothing asset, Lod Index and Mask index
-	TWeakObjectPtr<UClothingAsset> SelectedAsset;
+	TWeakObjectPtr<UClothingAssetCommon> SelectedAsset;
 	int32 SelectedLod;
 	int32 SelectedMask;
 
diff --git a/Engine/Source/Editor/ClothPainter/Private/SClothPaintTab.cpp b/Engine/Source/Editor/ClothPainter/Private/SClothPaintTab.cpp
--- a/Engine/Source/Editor/ClothPainter/Private/SClothPaintTab.cpp
+++ b/Engine/Source/Editor/ClothPainter/Private/SClothPaintTab.cpp
@@ -27,7 +27,7 @@
 #include "SClothPaintWidget.h"
 #include "IPersonaToolkit.h"
 #include "SClothAssetSelector.h"
-#include "Assets/ClothingAsset.h"
+#include "ClothingAsset.h"
 #include "Widgets/Layout/SScrollBox.h"
 #include "ComponentReregisterContext.h"
 
@@ -51,7 +51,7 @@ SClothPaintTab::~SClothPaintTab()
 
 void SClothPaintTab::Construct(const FArguments& InArgs)
 {
-	// Detail view for UClothingAsset
+	// Detail view for UClothingAssetCommon
 	FPropertyEditorModule& EditModule = FModuleManager::Get().GetModuleChecked<FPropertyEditorModule>("PropertyEditor");
 
 	FDetailsViewArgs DetailsViewArgs(
@@ -75,7 +75,7 @@ void SClothPaintTab::Construct(const FArguments& InArgs)
 
 	// Add the CDO by default
 	TArray<UObject*> Objects;
-	Objects.Add(UClothingAsset::StaticClass()->GetDefaultObject());
+	Objects.Add(UClothingAssetCommon::StaticClass()->GetDefaultObject());
 	DetailsView->SetObjects(Objects, true);
 
 	HostingApp = InArgs._InHostingApp;
@@ -147,7 +147,7 @@ void SClothPaintTab::UpdatePaintTools()
 		ISkeletalMeshEditor* SkeletalMeshEditor = static_cast<ISkeletalMeshEditor*>(HostingApp.Pin().Get());
 		SkeletalMeshEditor->GetAssetEditorModeManager()->ActivateMode(PaintModeID, true);
 
-		FClothingPaintEditMode* PaintMode = (FClothingPaintEditMode*)SkeletalMeshEditor->GetAssetEditorModeManager()->FindMode(PaintModeID);
+		FClothingPaintEditMode* PaintMode = (FClothingPaintEditMode*)SkeletalMeshEditor->GetAssetEditorModeManager()->GetActiveMode(PaintModeID);
 		if (PaintMode)
 		{
 			FClothPainter* ClothPainter = static_cast<FClothPainter*>(PaintMode->GetMeshPainter());
@@ -165,7 +165,7 @@ void SClothPaintTab::UpdatePaintTools()
 
 			if(SelectorWidget.IsValid())
 			{
-				TWeakObjectPtr<UClothingAsset> WeakAsset = SelectorWidget->GetSelectedAsset();
+				TWeakObjectPtr<UClothingAssetCommon> WeakAsset = SelectorWidget->GetSelectedAsset();
 
 				if(WeakAsset.Get())
 				{
@@ -183,13 +183,13 @@ void SClothPaintTab::UpdatePaintTools()
 	}
 }
 
-void SClothPaintTab::OnAssetSelectionChanged(TWeakObjectPtr<UClothingAsset> InAssetPtr, int32 InLodIndex, int32 InMaskIndex)
+void SClothPaintTab::OnAssetSelectionChanged(TWeakObjectPtr<UClothingAssetCommon> InAssetPtr, int32 InLodIndex, int32 InMaskIndex)
 {
 	if(bPaintModeEnabled)
 	{
 		ISkeletalMeshEditor* SkeletalMeshEditor = static_cast<ISkeletalMeshEditor*>(HostingApp.Pin().Get());
 
-		FClothingPaintEditMode* PaintMode = (FClothingPaintEditMode*)SkeletalMeshEditor->GetAssetEditorModeManager()->FindMode(PaintModeID);
+		FClothingPaintEditMode* PaintMode = (FClothingPaintEditMode*)SkeletalMeshEditor->GetAssetEditorModeManager()->GetActiveMode(PaintModeID);
 		if(PaintMode)
 		{
 			FClothPainter* ClothPainter = static_cast<FClothPainter*>(PaintMode->GetMeshPainter());
@@ -201,7 +201,7 @@ void SClothPaintTab::OnAssetSelectionChanged(TWeakObjectPtr<UClothingAsset> InAs
 		}
 	}
 
-	if(UClothingAsset* Asset = InAssetPtr.Get())
+	if(UClothingAssetCommon* Asset = InAssetPtr.Get())
 	{
 		TArray<UObject*> Objects;
 
@@ -220,7 +220,7 @@ bool SClothPaintTab::IsAssetDetailsPanelEnabled()
 
 		if(SelectedObjects.Num() > 0)
 		{
-			return SelectedObjects[0].Get() != UClothingAsset::StaticClass()->GetDefaultObject();
+			return SelectedObjects[0].Get() != UClothingAssetCommon::StaticClass()->GetDefaultObject();
 		}
 	}
 
diff --git a/Engine/Source/Editor/ClothPainter/Private/SClothPaintTab.h b/Engine/Source/Editor/ClothPainter/Private/SClothPaintTab.h
--- a/Engine/Source/Editor/ClothPainter/Private/SClothPaintTab.h
+++ b/Engine/Source/Editor/ClothPainter/Private/SClothPaintTab.h
@@ -4,7 +4,7 @@
 
 #include "Toolkits/AssetEditorToolkit.h"
 
-class UClothingAsset;
+class UClothingAssetCommon;
 class SClothPaintWidget;
 class SClothAssetSelector;
 class IDetailsView;
@@ -38,7 +38,7 @@ protected:
 	void UpdatePaintTools();
 
 	/** Called from the selector when the asset selection changes (Asset, LOD, Mask) */
-	void OnAssetSelectionChanged(TWeakObjectPtr<UClothingAsset> InAssetPtr, int32 InLodIndex, int32 InMaskIndex);
+	void OnAssetSelectionChanged(TWeakObjectPtr<UClothingAssetCommon> InAssetPtr, int32 InLodIndex, int32 InMaskIndex);
 
 	/** Whether or not the asset config section is enabled for editing */
 	bool IsAssetDetailsPanelEnabled();
diff --git a/Engine/Source/Editor/ClothPainter/Private/SClothPaintWidget.h b/Engine/Source/Editor/ClothPainter/Private/SClothPaintWidget.h
--- a/Engine/Source/Editor/ClothPainter/Private/SClothPaintWidget.h
+++ b/Engine/Source/Editor/ClothPainter/Private/SClothPaintWidget.h
@@ -9,9 +9,9 @@
 
 class IDetailsView;
 class FClothPainter;
-class UClothingAsset;
+class UClothingAssetCommon;
 class UClothPainterSettings;
-struct FClothParameterMask_PhysMesh;
+struct FPointWeightMap;
 
 class SClothPaintWidget : public SCompoundWidget
 {
diff --git a/Engine/Source/Editor/ClothPainter/Private/SCopyVertexColorSettingsPanel.cpp b/Engine/Source/Editor/ClothPainter/Private/SCopyVertexColorSettingsPanel.cpp
--- a/Engine/Source/Editor/ClothPainter/Private/SCopyVertexColorSettingsPanel.cpp
+++ b/Engine/Source/Editor/ClothPainter/Private/SCopyVertexColorSettingsPanel.cpp
@@ -8,12 +8,12 @@
 #include "Widgets/Layout/SUniformGridPanel.h"
 #include "Widgets/Layout/SBox.h"
 #include "Widgets/Input/SButton.h"
-#include "Assets/ClothingAsset.h"
+#include "ClothingAsset.h"
 #include "Framework/Application/SlateApplication.h"
 
 #define LOCTEXT_NAMESPACE "CopyVertexColorSettings"
 
-void SCopyVertexColorSettingsPanel::Construct(const FArguments& InArgs, UClothingAsset* InAsset, int32 InLOD, FClothParameterMask_PhysMesh* InMask)
+void SCopyVertexColorSettingsPanel::Construct(const FArguments& InArgs, UClothingAssetCommon* InAsset, int32 InLOD, FPointWeightMap* InMask)
 {
 	FPropertyEditorModule& PropertyEditorModule = FModuleManager::GetModuleChecked<FPropertyEditorModule>("PropertyEditor");
 
@@ -109,18 +109,18 @@ float GetColorChannelAsFloat(const FColor& Color, ESourceColorChannel Channel, f
 /** Copy channel of vertex colors into a particular mask entry */
 FReply SCopyVertexColorSettingsPanel::OnCopyClicked()
 {
-	UClothingAsset* Asset = SelectedAssetPtr.Get();
+	UClothingAssetCommon* Asset = SelectedAssetPtr.Get();
 	if (Asset && SelectedMask)
 	{
-		FClothLODData& ClothLODData = Asset->LodData[SelectedLOD];
+		UClothLODDataBase* ClothLODData = Asset->ClothLodData[SelectedLOD];
 
-		check(ClothLODData.PhysicalMeshData.Vertices.Num() == ClothLODData.PhysicalMeshData.VertexColors.Num());
-		int32 NumVerts = ClothLODData.PhysicalMeshData.Vertices.Num();
+		check(ClothLODData->PhysicalMeshData->Vertices.Num() == ClothLODData->PhysicalMeshData->VertexColors.Num());
+		int32 NumVerts = ClothLODData->PhysicalMeshData->Vertices.Num();
 		check(SelectedMask->Values.Num() == NumVerts);
 
 		for (int32 VertIdx = 0; VertIdx < NumVerts; VertIdx++)
 		{
-			const FColor VertColor = ClothLODData.PhysicalMeshData.VertexColors[VertIdx];
+			const FColor VertColor = ClothLODData->PhysicalMeshData->VertexColors[VertIdx];
 			SelectedMask->Values[VertIdx] = GetColorChannelAsFloat(VertColor, CopyParams.ColorChannel, CopyParams.ScalingFactor);
 		}
 	}
diff --git a/Engine/Source/Editor/ClothPainter/Private/SCopyVertexColorSettingsPanel.h b/Engine/Source/Editor/ClothPainter/Private/SCopyVertexColorSettingsPanel.h
--- a/Engine/Source/Editor/ClothPainter/Private/SCopyVertexColorSettingsPanel.h
+++ b/Engine/Source/Editor/ClothPainter/Private/SCopyVertexColorSettingsPanel.h
@@ -6,8 +6,8 @@
 #include "Widgets/SCompoundWidget.h"
 #include "Widgets/DeclarativeSyntaxSupport.h"
 
-class UClothingAsset;
-struct FClothParameterMask_PhysMesh;
+class UClothingAssetCommon;
+struct FPointWeightMap;
 
 /** Widget used for copying vertex colors from sim mesh to a selected mask. */
 class SCopyVertexColorSettingsPanel : public SCompoundWidget
@@ -17,7 +17,7 @@ public:
 	{}
 	SLATE_END_ARGS()
 
-	void Construct(const FArguments& InArgs, UClothingAsset* InAsset, int32 InLOD, FClothParameterMask_PhysMesh* InMask);
+	void Construct(const FArguments& InArgs, UClothingAssetCommon* InAsset, int32 InLOD, FPointWeightMap* InMask);
 
 private:
 
@@ -28,9 +28,9 @@ private:
 	FReply OnCopyClicked();
 
 	// Pointer to currently selected ClothingAsset
-	TWeakObjectPtr<UClothingAsset> SelectedAssetPtr;
+	TWeakObjectPtr<UClothingAssetCommon> SelectedAssetPtr;
 	// Pointer to selected mask
-	FClothParameterMask_PhysMesh* SelectedMask;
+	FPointWeightMap* SelectedMask;
 	// Currently selected LOD
 	int32 SelectedLOD;
 };
\ No newline at end of file
diff --git a/Engine/Source/Editor/ClothingSystemEditor/ClothingSystemEditor.Build.cs b/Engine/Source/Editor/ClothingSystemEditor/ClothingSystemEditor.Build.cs
--- a/Engine/Source/Editor/ClothingSystemEditor/ClothingSystemEditor.Build.cs
+++ b/Engine/Source/Editor/ClothingSystemEditor/ClothingSystemEditor.Build.cs
@@ -8,6 +8,7 @@ public class ClothingSystemEditor : ModuleRules
 	public ClothingSystemEditor(ReadOnlyTargetRules Target) : base(Target)
 	{
         PublicIncludePathModuleNames.Add("UnrealEd");
+        PublicIncludePathModuleNames.Add("ClothingSystemRuntimeInterface");
         PublicIncludePathModuleNames.Add("ClothingSystemRuntimeNv");
         PublicIncludePathModuleNames.Add("ClothingSystemEditorInterface");
 
@@ -23,6 +24,8 @@ public class ClothingSystemEditor : ModuleRules
 			new string[] { 
 				"Core",
 				"CoreUObject",
+				"ClothingSystemRuntimeInterface",
+                "ClothingSystemRuntimeCommon",
                 "ClothingSystemRuntimeNv",
                 "ContentBrowser",
                 "UnrealEd",
diff --git a/Engine/Source/Editor/ClothingSystemEditor/Private/ClothingAssetFactory.cpp b/Engine/Source/Editor/ClothingSystemEditor/Private/ClothingAssetFactory.cpp
--- a/Engine/Source/Editor/ClothingSystemEditor/Private/ClothingAssetFactory.cpp
+++ b/Engine/Source/Editor/ClothingSystemEditor/Private/ClothingAssetFactory.cpp
@@ -5,9 +5,11 @@
 #include "Factories.h"
 #include "PhysXPublic.h"
 #include "PhysicsPublic.h"
-#include "Assets/ClothingAsset.h"
+#include "ClothingAsset.h"
+#include "ClothPhysicalMeshDataNv.h" // for MaskTarget_PhysMesh
 #if WITH_APEX_CLOTHING
 #include "ClothingAssetAuthoring.h"
+#include "Assets/ClothingAssetNv.h"
 #endif // WITH_APEX_CLOTHING
 #include "ApexClothingUtils.h"
 #include "ContentBrowserModule.h"
@@ -93,7 +95,7 @@ UClothingAssetBase* UClothingAssetFactory::Import
 		return nullptr;
 	}
 
-	UClothingAsset* NewClothingAsset = nullptr;
+	UClothingAssetCommon* NewClothingAsset = nullptr;
 
 	TArray<uint8> FileBuffer;
 	if(FFileHelper::LoadFileToArray(FileBuffer, *Filename, FILEREAD_Silent))
@@ -107,7 +109,7 @@ UClothingAssetBase* UClothingAssetFactory::Import
 		}
 
 		// Create an unreal clothing asset
-		NewClothingAsset = Cast<UClothingAsset>(CreateFromApexAsset(ApexAsset, TargetMesh, InName));
+		NewClothingAsset = Cast<UClothingAssetCommon>(CreateFromApexAsset(ApexAsset, TargetMesh, InName));
 
 		if(NewClothingAsset)
 		{
@@ -150,8 +152,8 @@ UClothingAssetBase* UClothingAssetFactory::Reimport(const FString& Filename, USk
 	
 		FMultiComponentReregisterContext ReregisterContext(ComponentsToReregister);
 
-		UClothingAsset* OldClothingAsset = Cast<UClothingAsset>(TargetMesh->MeshClothingAssets[OldIndex]);
-		UClothingAsset* NewClothingAsset = nullptr;
+		UClothingAssetCommon* OldClothingAsset = Cast<UClothingAssetCommon>(TargetMesh->MeshClothingAssets[OldIndex]);
+		UClothingAssetCommon* NewClothingAsset = nullptr;
 		FName AssetName = NAME_None;
 
 		if(!OldClothingAsset || !TargetMesh->MeshClothingAssets.IsValidIndex(OldIndex))
@@ -172,7 +174,7 @@ UClothingAssetBase* UClothingAssetFactory::Reimport(const FString& Filename, USk
 			OldClothingAsset->UnbindFromSkeletalMesh(TargetMesh);
 
 			// Create an unreal clothing asset
-			NewClothingAsset = Cast<UClothingAsset>(CreateFromApexAsset(ApexAsset, TargetMesh, AssetName));
+			NewClothingAsset = Cast<UClothingAssetCommon>(CreateFromApexAsset(ApexAsset, TargetMesh, AssetName));
 
 			if(NewClothingAsset)
 			{
@@ -235,13 +237,18 @@ UClothingAssetBase* UClothingAssetFactory::CreateFromSkeletalMesh(USkeletalMesh*
 	}
 
 	FString SanitizedName = ObjectTools::SanitizeObjectName(Params.AssetName);
-	FName ObjectName = MakeUniqueObjectName(TargetMesh, UClothingAsset::StaticClass(), FName(*SanitizedName));
-	UClothingAsset* NewAsset = NewObject<UClothingAsset>(TargetMesh, ObjectName);
+	FName ObjectName = MakeUniqueObjectName(TargetMesh, UClothingAssetCommon::StaticClass(), FName(*SanitizedName));
+#if WITH_APEX_CLOTHING
+	UClothingAssetCommon* NewAsset = NewObject<UClothingAssetNv>(TargetMesh, ObjectName);
+#else
+	check(false); // TODO: Implement Chaos ClothingAsset, and have the factory get the correct object
+	UClothingAssetCommon* NewAsset = NewObject<UClothingAssetCommon>(TargetMesh, ObjectName);
+#endif
 	NewAsset->SetFlags(RF_Transactional);
 
 	// Adding a new LOD from this skeletal mesh
-	NewAsset->LodData.AddDefaulted();
-	FClothLODData& LodData = NewAsset->LodData.Last();
+	NewAsset->AddNewLod();
+	UClothLODDataBase* LodData = NewAsset->ClothLodData.Last();
 
 	if(ImportToLodInternal(TargetMesh, Params.LodIndex, Params.SourceSection, NewAsset, LodData))
 	{
@@ -271,7 +278,7 @@ UClothingAssetBase* UClothingAssetFactory::CreateFromSkeletalMesh(USkeletalMesh*
 
 UClothingAssetBase* UClothingAssetFactory::CreateFromExistingCloth(USkeletalMesh* TargetMesh, USkeletalMesh* SourceMesh, UClothingAssetBase* SourceAsset)
 {
-	UClothingAsset* SourceClothingAsset = Cast<UClothingAsset>(SourceAsset);
+	UClothingAssetCommon* SourceClothingAsset = Cast<UClothingAssetCommon>(SourceAsset);
 
 	if (!SourceClothingAsset)
 	{
@@ -279,7 +286,7 @@ UClothingAssetBase* UClothingAssetFactory::CreateFromExistingCloth(USkeletalMesh
 	}
 
 	//Duplicating the clothing asset using the existing asset as a template
-	UClothingAsset* NewAsset = DuplicateObject<UClothingAsset>(SourceClothingAsset, TargetMesh, SourceClothingAsset->GetFName());
+	UClothingAssetCommon* NewAsset = DuplicateObject<UClothingAssetCommon>(SourceClothingAsset, TargetMesh, SourceClothingAsset->GetFName());
 
 	NewAsset->AssetGuid = FGuid::NewGuid();
 	//Need to empty LODMap to remove previous mappings from cloth LOD to SkelMesh LOD
@@ -318,40 +325,40 @@ UClothingAssetBase* UClothingAssetFactory::ImportLodToClothing(USkeletalMesh* Ta
 		if(TargetMesh->MeshClothingAssets.Find(TargetClothing, MeshAssetIndex))
 		{
 			// Everything looks good, continue to actual import
-			UClothingAsset* ConcreteTarget = CastChecked<UClothingAsset>(TargetClothing);
+			UClothingAssetCommon* ConcreteTarget = CastChecked<UClothingAssetCommon>(TargetClothing);
 
-			FClothLODData* RemapSource = nullptr;
+			UClothLODDataBase* RemapSource = nullptr;
 
 			if(Params.bRemapParameters)
 			{
-				if(Params.TargetLod == ConcreteTarget->LodData.Num())
+				if(Params.TargetLod == ConcreteTarget->GetNumLods())
 				{
 					// New LOD, remap from previous
-					RemapSource = &ConcreteTarget->LodData.Last();
+					RemapSource = ConcreteTarget->ClothLodData.Last();
 				}
 				else
 				{
 					// This is a replacement, remap from current LOD
-					check(ConcreteTarget->LodData.IsValidIndex(Params.TargetLod));
-					RemapSource = &ConcreteTarget->LodData[Params.TargetLod];
+					check(ConcreteTarget->ClothLodData.IsValidIndex(Params.TargetLod));
+					RemapSource = ConcreteTarget->ClothLodData[Params.TargetLod];
 				}
 			}
 
-			if(Params.TargetLod == ConcreteTarget->LodData.Num())
+			if(Params.TargetLod == ConcreteTarget->GetNumLods())
 			{
-				ConcreteTarget->LodData.AddDefaulted();
+				ConcreteTarget->AddNewLod();
 			}
-			else if(!ConcreteTarget->LodData.IsValidIndex(Params.TargetLod))
+			else if(!ConcreteTarget->ClothLodData.IsValidIndex(Params.TargetLod))
 			{
 				LogAndToastWarning(LOCTEXT("Warning_InvalidLodTarget", "Failed to import clothing LOD, invalid target LOD."));
 				return nullptr;
 			}
 
-			FClothLODData& NewLod = ConcreteTarget->LodData[Params.TargetLod];
+			UClothLODDataBase* NewLod = ConcreteTarget->ClothLodData[Params.TargetLod];
 
 			if(Params.TargetLod > 0 && Params.bRemapParameters)
 			{
-				RemapSource = &ConcreteTarget->LodData[Params.TargetLod - 1];
+				RemapSource = ConcreteTarget->ClothLodData[Params.TargetLod - 1];
 			}
 
 			if(ImportToLodInternal(TargetMesh, Params.LodIndex, Params.SourceSection, ConcreteTarget, NewLod, RemapSource))
@@ -382,7 +389,7 @@ UClothingAssetBase* UClothingAssetFactory::ImportLodToClothing(USkeletalMesh* Ta
 UClothingAssetBase* UClothingAssetFactory::CreateFromApexAsset(nvidia::apex::ClothingAsset* InApexAsset, USkeletalMesh* TargetMesh, FName InName)
 {
 #if WITH_APEX_CLOTHING
-	UClothingAsset* NewClothingAsset = NewObject<UClothingAsset>(TargetMesh, InName);
+	UClothingAssetNv* NewClothingAsset = NewObject<UClothingAssetNv>(TargetMesh, InName);
 	NewClothingAsset->SetFlags(RF_Transactional);
 
 	const NvParameterized::Interface* AssetParams = InApexAsset->getAssetNvParameterized();
@@ -393,11 +400,15 @@ UClothingAssetBase* UClothingAssetFactory::CreateFromApexAsset(nvidia::apex::Clo
 	int32 NumLodsToBuild = 0;
 	GraphicalLodArrayHandle.getArraySize(NumLodsToBuild);
 
-	NewClothingAsset->LodData.AddZeroed(NumLodsToBuild);
-
 	for(int32 CurrLodIdx = 0; CurrLodIdx < NumLodsToBuild; ++CurrLodIdx)
 	{
-		FClothLODData& CurrentLodData = NewClothingAsset->LodData[CurrLodIdx];
+		NewClothingAsset->AddNewLod();
+		UClothLODDataNv* CurrentLodDataPtr = Cast<UClothLODDataNv>(NewClothingAsset->ClothLodData[CurrLodIdx]);
+		if (!ensure(CurrentLodDataPtr))
+		{
+			continue;
+		}
+		UClothLODDataNv& CurrentLodData = *CurrentLodDataPtr;
 
 		TArray<FApexVertData> ApexVertData;
 
@@ -406,14 +417,19 @@ UClothingAssetBase* UClothingAssetFactory::CreateFromApexAsset(nvidia::apex::Clo
 		ExtractMaterialParameters(NewClothingAsset, *InApexAsset);
 
 		// Set to use legacy wind calculations, which is what APEX would normally have used
-		NewClothingAsset->ClothConfig.WindMethod = EClothingWindMethod::Legacy;
+		UClothConfigNv* ConfigNv = Cast<UClothConfigNv>(NewClothingAsset->ClothSimConfig);
+		if (!ensure(ConfigNv))
+		{
+			continue;
+		}
+		ConfigNv->WindMethod = EClothingWindMethod::Legacy;
 
 		// Fixup unreal-side bone indices
-		const int32 NumBoneDatas = CurrentLodData.PhysicalMeshData.BoneData.Num();
+		const int32 NumBoneDatas = CurrentLodData.PhysicalMeshData->BoneData.Num();
 		check(NumBoneDatas == ApexVertData.Num());
 		for(int32 BoneDataIndex = 0; BoneDataIndex < NumBoneDatas; ++BoneDataIndex)
 		{
-			FClothVertBoneData& BoneData = CurrentLodData.PhysicalMeshData.BoneData[BoneDataIndex];
+			FClothVertBoneData& BoneData = CurrentLodData.PhysicalMeshData->BoneData[BoneDataIndex];
 			FApexVertData& CurrentVertData = ApexVertData[BoneDataIndex];
 		
 			for(int32 BoneInfluenceIdx = 0; BoneInfluenceIdx < MAX_TOTAL_INFLUENCES; ++BoneInfluenceIdx)
@@ -460,32 +476,38 @@ UClothingAssetBase* UClothingAssetFactory::CreateFromApexAsset(nvidia::apex::Clo
 	NewClothingAsset->CalculateReferenceBoneIndex();
 
 	// Add masks for parameters
-	for(FClothLODData& Lod : NewClothingAsset->LodData)
+	for (UClothLODDataBase* LodPtr : NewClothingAsset->ClothLodData)
 	{
-		FClothPhysicalMeshData& PhysMesh = Lod.PhysicalMeshData;
+		check(LodPtr);
+		UClothPhysicalMeshDataNv* PhysMeshPtr = Cast<UClothPhysicalMeshDataNv>(LodPtr->PhysicalMeshData);
+		if (!ensure(PhysMeshPtr))
+		{
+			continue;
+		}
+		UClothPhysicalMeshDataNv& PhysMesh = *PhysMeshPtr;
 
 		// Didn't do anything previously - clear out incase there's something in there
 		// so we can use it correctly now.
-		Lod.ParameterMasks.Reset(3);
+		LodPtr->ParameterMasks.Reset(3);
 
 		// Max distances
-		Lod.ParameterMasks.AddDefaulted();
-		FClothParameterMask_PhysMesh& MaxDistanceMask = Lod.ParameterMasks.Last();
-		MaxDistanceMask.CopyFromPhysMesh(PhysMesh, MaskTarget_PhysMesh::MaxDistance);
+		LodPtr->ParameterMasks.AddDefaulted();
+		FPointWeightMap& MaxDistanceMask = LodPtr->ParameterMasks.Last();
+		MaxDistanceMask.CopyFrom(PhysMesh.GetFloatArray((int32)MaskTarget_PhysMesh::MaxDistance), (uint8)MaskTarget_PhysMesh::MaxDistance);
 		MaxDistanceMask.bEnabled = true;
 
 		if(PhysMesh.BackstopRadiuses.FindByPredicate([](const float& A) {return A != 0.0f; }))
 		{
 			// Backstop radii
-			Lod.ParameterMasks.AddDefaulted();
-			FClothParameterMask_PhysMesh& BackstopRadiusMask = Lod.ParameterMasks.Last();
-			BackstopRadiusMask.CopyFromPhysMesh(PhysMesh, MaskTarget_PhysMesh::BackstopRadius);
+			LodPtr->ParameterMasks.AddDefaulted();
+			FPointWeightMap& BackstopRadiusMask = LodPtr->ParameterMasks.Last();
+			BackstopRadiusMask.CopyFrom(PhysMesh.GetFloatArray((int32)MaskTarget_PhysMesh::BackstopRadius), (uint8)MaskTarget_PhysMesh::BackstopRadius);
 			BackstopRadiusMask.bEnabled = true;
 
 			// Backstop distances
-			Lod.ParameterMasks.AddDefaulted();
-			FClothParameterMask_PhysMesh& BackstopDistanceMask = Lod.ParameterMasks.Last();
-			BackstopDistanceMask.CopyFromPhysMesh(PhysMesh, MaskTarget_PhysMesh::BackstopDistance);
+			LodPtr->ParameterMasks.AddDefaulted();
+			FPointWeightMap& BackstopDistanceMask = LodPtr->ParameterMasks.Last();
+			BackstopDistanceMask.CopyFrom(PhysMesh.GetFloatArray((int32)MaskTarget_PhysMesh::BackstopDistance), (uint8)MaskTarget_PhysMesh::BackstopDistance);
 			BackstopDistanceMask.bEnabled = true;
 		}
 	}
@@ -821,7 +843,7 @@ void UClothingAssetFactory::FlipAuthoringUvs(NvParameterized::Interface* InRende
 	}
 }
 
-void UClothingAssetFactory::ExtractBoneData(UClothingAsset* NewAsset, ClothingAsset &InApexAsset)
+void UClothingAssetFactory::ExtractBoneData(UClothingAssetCommon* NewAsset, ClothingAsset &InApexAsset)
 {
 	const uint32 NumApexUsedBones = InApexAsset.getNumUsedBones();
 
@@ -834,7 +856,7 @@ void UClothingAssetFactory::ExtractBoneData(UClothingAsset* NewAsset, ClothingAs
 	}
 }
 
-void UClothingAssetFactory::ExtractSphereCollisions(UClothingAsset* NewAsset, nvidia::apex::ClothingAsset &InApexAsset, int32 InLodIdx, FClothLODData &InLodData)
+void UClothingAssetFactory::ExtractSphereCollisions(UClothingAssetCommon* NewAsset, nvidia::apex::ClothingAsset &InApexAsset, int32 InLodIdx, UClothLODDataNv &InLodData)
 {
 	const NvParameterized::Interface* AssetParams = InApexAsset.getAssetNvParameterized();
 
@@ -966,7 +988,7 @@ void UClothingAssetFactory::ExtractSphereCollisions(UClothingAsset* NewAsset, nv
 	}
 }
 
-void UClothingAssetFactory::ExtractMaterialParameters(UClothingAsset* NewAsset, nvidia::apex::ClothingAsset &InApexAsset)
+void UClothingAssetFactory::ExtractMaterialParameters(UClothingAssetCommon* NewAsset, nvidia::apex::ClothingAsset &InApexAsset)
 {
 	const NvParameterized::Interface* AssetParams = InApexAsset.getAssetNvParameterized();
 
@@ -986,8 +1008,9 @@ void UClothingAssetFactory::ExtractMaterialParameters(UClothingAsset* NewAsset,
 
 	MaterialArrayHandle.set(MaterialIndex);
 
+	if(UClothConfigNv* ConfigPtr = Cast<UClothConfigNv>(NewAsset->ClothSimConfig))
 	{
-		FClothConfig& Config = NewAsset->ClothConfig;
+		UClothConfigNv& Config = *ConfigPtr;
 		NvParameterized::Handle ChildHandle(MaterialArrayHandle);
 
 		// Read out material params
@@ -1111,7 +1134,13 @@ void UClothingAssetFactory::ExtractMaterialParameters(UClothingAsset* NewAsset,
 #endif // WITH_APEX_CLOTHING
 
 
-bool UClothingAssetFactory::ImportToLodInternal(USkeletalMesh* SourceMesh, int32 SourceLodIndex, int32 SourceSectionIndex, UClothingAsset* DestAsset, FClothLODData& DestLod, FClothLODData* InParameterRemapSource)
+bool UClothingAssetFactory::ImportToLodInternal(
+	USkeletalMesh* SourceMesh, 
+	int32 SourceLodIndex, 
+	int32 SourceSectionIndex, 
+	UClothingAssetCommon* DestAsset, 
+	UClothLODDataBase* DestLod, 
+	UClothLODDataBase* InParameterRemapSource)
 {
 	if(!SourceMesh || !SourceMesh->GetImportedModel())
 	{
@@ -1149,8 +1178,7 @@ bool UClothingAssetFactory::ImportToLodInternal(USkeletalMesh* SourceMesh, int32
 	TArray<uint32> IndexRemap;
 	IndexRemap.AddDefaulted(NumVerts);
 	{
-		float ThreshSq = SMALL_NUMBER * SMALL_NUMBER;
-
+		static const float ThreshSq = SMALL_NUMBER * SMALL_NUMBER;
 		for(int32 VertIndex = 0; VertIndex < NumVerts; ++VertIndex)
 		{
 			const FSoftSkinVertex& SourceVert = SourceSection.SoftVertices[VertIndex];
@@ -1196,13 +1224,14 @@ bool UClothingAssetFactory::ImportToLodInternal(USkeletalMesh* SourceMesh, int32
 	TArray<FVector> CachedNormals;
 	TArray<uint32> CachedIndices;
 	int32 NumSourceMasks = 0;
-	TArray<FClothParameterMask_PhysMesh> SourceMaskCopy;
+	TArray<FPointWeightMap> SourceMaskCopy;
 	
 	bool bPerformParamterRemap = false;
 
 	if(InParameterRemapSource)
 	{
-		FClothPhysicalMeshData& RemapPhysMesh = InParameterRemapSource->PhysicalMeshData;
+		check(InParameterRemapSource->PhysicalMeshData);
+		UClothPhysicalMeshDataBase& RemapPhysMesh = *InParameterRemapSource->PhysicalMeshData;
 		CachedPositions = RemapPhysMesh.Vertices;
 		CachedNormals = RemapPhysMesh.Normals;
 		CachedIndices = RemapPhysMesh.Indices;
@@ -1212,11 +1241,14 @@ bool UClothingAssetFactory::ImportToLodInternal(USkeletalMesh* SourceMesh, int32
 		bPerformParamterRemap = true;
 	}
 
-	FClothPhysicalMeshData& PhysMesh = DestLod.PhysicalMeshData;
+	check(DestLod->PhysicalMeshData);
+	UClothPhysicalMeshDataBase& PhysMesh = *DestLod->PhysicalMeshData;
 	PhysMesh.Reset(NumUniqueVerts);
 	PhysMesh.Indices.Reset(NumIndices);
 	PhysMesh.Indices.AddZeroed(NumIndices);
 
+	TArray<TArray<float>*> WeightArrays = PhysMesh.GetFloatArrays();
+
 	for(int32 VertexIndex = 0; VertexIndex < NumUniqueVerts; ++VertexIndex)
 	{
 		const FSoftSkinVertex& SourceVert = SourceSection.SoftVertices[OriginalIndexes[VertexIndex]];
@@ -1224,16 +1256,20 @@ bool UClothingAssetFactory::ImportToLodInternal(USkeletalMesh* SourceMesh, int32
 		PhysMesh.Vertices[VertexIndex] = SourceVert.Position;
 		PhysMesh.Normals[VertexIndex] = SourceVert.TangentZ;
 		PhysMesh.VertexColors[VertexIndex] = SourceVert.Color;
-		PhysMesh.MaxDistances[VertexIndex] = 0.0f;
 
-		PhysMesh.BackstopRadiuses[VertexIndex] = 0.0f;
-		PhysMesh.BackstopDistances[VertexIndex] = 0.0f;
+		for (int32 WeightArrayIndex = 0; WeightArrayIndex < WeightArrays.Num(); ++WeightArrayIndex)
+		{
+			TArray<float>* Weights = WeightArrays[WeightArrayIndex];
+			if (Weights && Weights->IsValidIndex(VertexIndex))
+			{
+				(*Weights)[VertexIndex] = 0.0f;
+			}
+		}
 
 		FClothVertBoneData& BoneData = PhysMesh.BoneData[VertexIndex];
 		for(int32 InfluenceIndex = 0; InfluenceIndex < MAX_TOTAL_INFLUENCES; ++InfluenceIndex)
 		{
 			const uint16 SourceIndex = SourceSection.BoneMap[SourceVert.InfluenceBones[InfluenceIndex]];
-
 			if(SourceIndex != INDEX_NONE)
 			{
 				FName BoneName = SourceMesh->RefSkeleton.GetBoneName(SourceIndex);
@@ -1244,9 +1280,9 @@ bool UClothingAssetFactory::ImportToLodInternal(USkeletalMesh* SourceMesh, int32
 	}
 
 	// Add a max distance parameter mask to begin with
-	DestLod.ParameterMasks.AddDefaulted();
-	FClothParameterMask_PhysMesh& Mask = DestLod.ParameterMasks.Last();
-	Mask.CopyFromPhysMesh(PhysMesh, MaskTarget_PhysMesh::MaxDistance);
+	DestLod->ParameterMasks.AddDefaulted();
+	FPointWeightMap& Mask = DestLod->ParameterMasks.Last();
+	Mask.CopyFrom(PhysMesh.GetFloatArray((uint32)MaskTarget_PhysMesh::MaxDistance), (uint8)MaskTarget_PhysMesh::MaxDistance);
 	Mask.bEnabled = true;
 
 	PhysMesh.MaxBoneWeights = SourceSection.MaxBoneInfluences;
@@ -1279,16 +1315,16 @@ bool UClothingAssetFactory::ImportToLodInternal(USkeletalMesh* SourceMesh, int32
 	{
 		ClothingMeshUtils::FVertexParameterMapper ParameterRemapper(PhysMesh.Vertices, PhysMesh.Normals, CachedPositions, CachedNormals, CachedIndices);
 
-		DestLod.ParameterMasks.Reset(NumSourceMasks);
+		DestLod->ParameterMasks.Reset(NumSourceMasks);
 
 		for(int32 MaskIndex = 0; MaskIndex < NumSourceMasks; ++MaskIndex)
 		{
-			FClothParameterMask_PhysMesh& SourceMask = SourceMaskCopy[MaskIndex];
+			FPointWeightMap& SourceMask = SourceMaskCopy[MaskIndex];
 
-			DestLod.ParameterMasks.AddDefaulted();
-			FClothParameterMask_PhysMesh& DestMask = DestLod.ParameterMasks.Last();
+			DestLod->ParameterMasks.AddDefaulted();
+			FPointWeightMap& DestMask = DestLod->ParameterMasks.Last();
 
-			DestMask.Initialize(PhysMesh);
+			DestMask.Initialize(PhysMesh.Vertices.Num());
 			DestMask.CurrentTarget = SourceMask.CurrentTarget;
 			DestMask.bEnabled = SourceMask.bEnabled;
 
@@ -1303,11 +1339,16 @@ bool UClothingAssetFactory::ImportToLodInternal(USkeletalMesh* SourceMesh, int32
 
 #if WITH_APEX_CLOTHING
 
-void UClothingAssetFactory::ExtractLodPhysicalData(UClothingAsset* NewAsset, ClothingAsset &InApexAsset, int32 InLodIdx, FClothLODData &InLodData, TArray<FApexVertData>& OutApexVertData)
+void UClothingAssetFactory::ExtractLodPhysicalData(UClothingAssetCommon* NewAsset, ClothingAsset &InApexAsset, int32 InLodIdx, UClothLODDataNv &InLodData, TArray<FApexVertData>& OutApexVertData)
 {
 	const NvParameterized::Interface* AssetParams = InApexAsset.getAssetNvParameterized();
 
-	FClothPhysicalMeshData& PhysData = InLodData.PhysicalMeshData;
+	UClothPhysicalMeshDataNv* PhysDataPtr = Cast<UClothPhysicalMeshDataNv>(InLodData.PhysicalMeshData);
+	if (!ensure(PhysDataPtr))
+	{
+		return;
+	}
+	UClothPhysicalMeshDataNv& PhysData = *PhysDataPtr;
 
 	NvParameterized::Handle GraphicalMeshArrayHandle(*AssetParams, "graphicalLods");
 
diff --git a/Engine/Source/Editor/ClothingSystemEditor/Public/ClothingAssetFactory.h b/Engine/Source/Editor/ClothingSystemEditor/Public/ClothingAssetFactory.h
--- a/Engine/Source/Editor/ClothingSystemEditor/Public/ClothingAssetFactory.h
+++ b/Engine/Source/Editor/ClothingSystemEditor/Public/ClothingAssetFactory.h
@@ -2,8 +2,9 @@
 
 #pragma once
 
-#include "Assets/ClothingAsset.h"
+#include "ClothingAsset.h"
 #include "ClothingAssetFactoryInterface.h"
+#include "ClothLODDataNv.h"
 #include "GPUSkinPublicDefs.h"
 
 #include "ClothingAssetFactory.generated.h"
@@ -44,7 +45,7 @@ public:
 	// Tests whether the given filename should be able to be imported
 	virtual bool CanImport(const FString& Filename) override;
 
-	// Given an APEX asset, build a UClothingAsset containing the required data
+	// Given an APEX asset, build a UClothingAssetCommon containing the required data
 	virtual UClothingAssetBase* CreateFromApexAsset(nvidia::apex::ClothingAsset* InApexAsset, USkeletalMesh* TargetMesh, FName InName = NAME_None) override;
 
 private:
@@ -63,17 +64,17 @@ private:
 	void FlipAuthoringUvs(NvParameterized::Interface* InRenderMeshAuthoringInterface, bool bFlipU, bool bFlipV);
 
 	// Extraction methods for pulling the required data from an APEX asset and
-	// pushing it to a UClothingAsset
-	void ExtractLodPhysicalData(UClothingAsset* NewAsset, nvidia::apex::ClothingAsset &InApexAsset, int32 InLodIdx, FClothLODData &InLodData, TArray<FApexVertData>& OutApexVertData);
-	void ExtractBoneData(UClothingAsset* NewAsset, nvidia::apex::ClothingAsset &InApexAsset);
-	void ExtractSphereCollisions(UClothingAsset* NewAsset, nvidia::apex::ClothingAsset &InApexAsset, int32 InLodIdx, FClothLODData &InLodData);
-	void ExtractMaterialParameters(UClothingAsset* NewAsset, nvidia::apex::ClothingAsset &InApexAsset);
+	// pushing it to a UClothingAssetCommon
+	void ExtractLodPhysicalData(UClothingAssetCommon* NewAsset, nvidia::apex::ClothingAsset &InApexAsset, int32 InLodIdx, UClothLODDataNv &InLodData, TArray<FApexVertData>& OutApexVertData);
+	void ExtractBoneData(UClothingAssetCommon* NewAsset, nvidia::apex::ClothingAsset &InApexAsset);
+	void ExtractSphereCollisions(UClothingAssetCommon* NewAsset, nvidia::apex::ClothingAsset &InApexAsset, int32 InLodIdx, UClothLODDataNv &InLodData);
+	void ExtractMaterialParameters(UClothingAssetCommon* NewAsset, nvidia::apex::ClothingAsset &InApexAsset);
 #endif
 
 	// Utility methods for skeletal mesh extraction //////////////////////////
 
 	/** Handles internal import of LODs */
-	bool ImportToLodInternal(USkeletalMesh* SourceMesh, int32 SourceLodIndex, int32 SourceSectionIndex, UClothingAsset* DestAsset, FClothLODData& DestLod, FClothLODData* InParameterRemapSource = nullptr);
+	bool ImportToLodInternal(USkeletalMesh* SourceMesh, int32 SourceLodIndex, int32 SourceSectionIndex, UClothingAssetCommon* DestAsset, UClothLODDataBase* DestLod, UClothLODDataBase* InParameterRemapSource = nullptr);
 
 	//////////////////////////////////////////////////////////////////////////
 
diff --git a/Engine/Source/Editor/ContentBrowser/Private/SFilterList.cpp b/Engine/Source/Editor/ContentBrowser/Private/SFilterList.cpp
--- a/Engine/Source/Editor/ContentBrowser/Private/SFilterList.cpp
+++ b/Engine/Source/Editor/ContentBrowser/Private/SFilterList.cpp
@@ -26,6 +26,16 @@
 
 #define LOCTEXT_NAMESPACE "ContentBrowser"
 
+/** Helper struct to avoid friending the whole of SFilterList */
+struct FFrontendFilterExternalActivationHelper
+{
+	static void BindToFilter(TSharedRef<SFilterList> InFilterList, TSharedRef<FFrontendFilter> InFrontendFilter)
+	{
+		TWeakPtr<FFrontendFilter> WeakFilter = InFrontendFilter;
+		InFrontendFilter->SetActiveEvent.AddSP(&InFilterList.Get(), &SFilterList::OnSetFilterActive, WeakFilter);
+	}
+};
+
 /** A class for check boxes in the filter list. If you double click a filter checkbox, you will enable it and disable all others */
 class SFilterCheckBox : public SCheckBox
 {
@@ -516,10 +526,15 @@ void SFilterList::Construct( const FArguments& InArgs )
 
 	AllFrontendFilterCategories.AddUnique(DefaultCategory);
 
-	// Auto add all inverse filters
+	
 	for (auto FilterIt = AllFrontendFilters.CreateConstIterator(); FilterIt; ++FilterIt)
 	{
-		SetFrontendFilterActive(*FilterIt, false);
+		// Bind external activation event
+		TSharedRef<FFrontendFilter> Filter = *FilterIt;
+		FFrontendFilterExternalActivationHelper::BindToFilter(SharedThis(this), Filter);
+
+		// Auto add all inverse filters
+		SetFrontendFilterActive(Filter, false);
 	}
 
 	FilterBox = SNew(SWrapBox)
@@ -1483,4 +1498,17 @@ void SFilterList::OnResetFilters()
 	RemoveAllFilters();
 }
 
+void SFilterList::OnSetFilterActive(bool bInActive, TWeakPtr<FFrontendFilter> InWeakFilter)
+{
+	TSharedPtr<FFrontendFilter> Filter = InWeakFilter.Pin();
+	if (Filter.IsValid())
+	{
+		if (!IsFrontendFilterInUse(Filter.ToSharedRef()))
+		{
+			TSharedRef<SFilter> NewFilter = AddFilter(Filter.ToSharedRef());
+			NewFilter->SetEnabled(bInActive);
+		}
+	}
+}
+
 #undef LOCTEXT_NAMESPACE
diff --git a/Engine/Source/Editor/ContentBrowser/Private/SFilterList.h b/Engine/Source/Editor/ContentBrowser/Private/SFilterList.h
--- a/Engine/Source/Editor/ContentBrowser/Private/SFilterList.h
+++ b/Engine/Source/Editor/ContentBrowser/Private/SFilterList.h
@@ -131,6 +131,9 @@ private:
 	/** Called when reset filters option is pressed */
 	void OnResetFilters();
 
+	/** Called to set a filter active externally */
+	void OnSetFilterActive(bool bInActive, TWeakPtr<FFrontendFilter> InWeakFilter);
+
 private:
 	/** The horizontal box which contains all the filters */
 	TSharedPtr<SWrapBox> FilterBox;
@@ -155,4 +158,6 @@ private:
 
 	/** Delegate for when filters have changed */
 	FOnFilterChanged OnFilterChanged;
+
+	friend struct FFrontendFilterExternalActivationHelper;
 };
diff --git a/Engine/Source/Editor/ContentBrowser/Public/FrontendFilterBase.h b/Engine/Source/Editor/ContentBrowser/Public/FrontendFilterBase.h
--- a/Engine/Source/Editor/ContentBrowser/Public/FrontendFilterBase.h
+++ b/Engine/Source/Editor/ContentBrowser/Public/FrontendFilterBase.h
@@ -58,6 +58,9 @@ public:
 	/** Called when the state of a particular Content Browser is being loaded from INI */
 	virtual void LoadSettings(const FString& IniFilename, const FString& IniSection, const FString& SettingsString) {}
 
+	/** Set this filter as active/inactive */
+	void SetActive(bool bInActive) { SetActiveEvent.Broadcast(bInActive); }
+
 	// IFilter implementation
 	DECLARE_DERIVED_EVENT( FFrontendFilter, IFilter<FAssetFilterType>::FChangedEvent, FChangedEvent );
 	virtual FChangedEvent& OnChanged() override { return ChangedEvent; }
@@ -70,5 +73,11 @@ protected:
 private:
 	FChangedEvent ChangedEvent;
 
+	/** This event is broadcast to set this filter active in the content browser it is being used in */
+	DECLARE_EVENT_OneParam(FFrontendFilter, FSetActiveEvent, bool);
+	FSetActiveEvent SetActiveEvent;
+
 	TSharedPtr<FFrontendFilterCategory> FilterCategory;
+
+	friend struct FFrontendFilterExternalActivationHelper;
 };
diff --git a/Engine/Source/Editor/DetailCustomizations/Private/DetailCustomizations.cpp b/Engine/Source/Editor/DetailCustomizations/Private/DetailCustomizations.cpp
--- a/Engine/Source/Editor/DetailCustomizations/Private/DetailCustomizations.cpp
+++ b/Engine/Source/Editor/DetailCustomizations/Private/DetailCustomizations.cpp
@@ -55,7 +55,6 @@
 #include "BlackboardEntryDetails.h"
 #include "AIDataProviderValueDetails.h"
 #include "EnvQueryParamInstanceCustomization.h"
-#include "SkeletonNotifyDetails.h"
 #include "Customizations/ColorStructCustomization.h"
 #include "SlateColorCustomization.h"
 #include "CurveStructCustomization.h"
@@ -317,7 +316,6 @@ void FDetailCustomizationsModule::RegisterObjectCustomizations()
 
 	RegisterCustomClassLayout("EditorAnimSegment", FOnGetDetailCustomizationInstance::CreateStatic(&FAnimMontageSegmentDetails::MakeInstance));
 	RegisterCustomClassLayout("EditorAnimCompositeSegment", FOnGetDetailCustomizationInstance::CreateStatic(&FAnimMontageSegmentDetails::MakeInstance));
-	RegisterCustomClassLayout("EditorSkeletonNotifyObj", FOnGetDetailCustomizationInstance::CreateStatic(&FSkeletonNotifyDetails::MakeInstance));
 	RegisterCustomClassLayout("AnimStateNode", FOnGetDetailCustomizationInstance::CreateStatic(&FAnimStateNodeDetails::MakeInstance));
 	RegisterCustomClassLayout("AnimStateTransitionNode", FOnGetDetailCustomizationInstance::CreateStatic(&FAnimTransitionNodeDetails::MakeInstance));
 	RegisterCustomClassLayout("AnimGraphNode_Trail", FOnGetDetailCustomizationInstance::CreateStatic(&FAnimTrailNodeDetails::MakeInstance));
diff --git a/Engine/Source/Editor/KismetCompiler/Private/AnimBlueprintCompiler.cpp b/Engine/Source/Editor/KismetCompiler/Private/AnimBlueprintCompiler.cpp
--- a/Engine/Source/Editor/KismetCompiler/Private/AnimBlueprintCompiler.cpp
+++ b/Engine/Source/Editor/KismetCompiler/Private/AnimBlueprintCompiler.cpp
@@ -267,7 +267,7 @@ UK2Node_CallFunction* FAnimBlueprintCompilerContext::SpawnCallAnimInstanceFuncti
 	return FunctionCall;
 }
 
-void FAnimBlueprintCompilerContext::CreateEvaluationHandlerStruct(UAnimGraphNode_Base* VisualAnimNode, FEvaluationHandlerRecord& Record)
+void FAnimBlueprintCompilerContext::CreateEvaluationHandler(UAnimGraphNode_Base* VisualAnimNode, FEvaluationHandlerRecord& Record)
 {
 	// Shouldn't create a handler if there is nothing to work with
 	check(Record.ServicedProperties.Num() > 0);
@@ -286,7 +286,7 @@ void FAnimBlueprintCompilerContext::CreateEvaluationHandlerStruct(UAnimGraphNode
 	// check function name isnt already used (data exists that can contain duplicate GUIDs) and apply a numeric extension until it is unique
 	int32 ExtensionIndex = 0;
 	FName* ExistingName = HandlerFunctionNames.Find(Record.HandlerFunctionName);
-	while(ExistingName != nullptr)
+	while (ExistingName != nullptr)
 	{
 		FunctionName = FString::Printf(TEXT("%s_%s_%s_%s_%d"), *AnimGraphDefaultSchema->DefaultEvaluationHandlerName.ToString(), *VisualAnimNode->GetOuter()->GetName(), *VisualAnimNode->GetClass()->GetName(), *VisualAnimNode->NodeGuid.ToString(), ExtensionIndex);
 		Record.HandlerFunctionName = FName(*FunctionName);
@@ -295,7 +295,7 @@ void FAnimBlueprintCompilerContext::CreateEvaluationHandlerStruct(UAnimGraphNode
 	}
 
 	HandlerFunctionNames.Add(Record.HandlerFunctionName);
-	
+
 	// Add a custom event in the graph
 	UK2Node_CustomEvent* EntryNode = SpawnIntermediateEventNode<UK2Node_CustomEvent>(VisualAnimNode, nullptr, ConsolidatedEventGraph);
 	EntryNode->bInternalEvent = true;
@@ -304,187 +304,159 @@ void FAnimBlueprintCompilerContext::CreateEvaluationHandlerStruct(UAnimGraphNode
 
 	// The ExecChain is the current exec output pin in the linear chain
 	UEdGraphPin* ExecChain = Schema->FindExecutionPin(*EntryNode, EGPD_Output);
+	if (Record.bServicesInstanceProperties)
+	{
+		// Need to create a variable set call for each serviced property in the handler
+		for (TPair<FName, FAnimNodeSinglePropertyHandler>& PropHandlerPair : Record.ServicedProperties)
+		{
+			FAnimNodeSinglePropertyHandler& PropHandler = PropHandlerPair.Value;
+			FName PropertyName = PropHandlerPair.Key;
 
-	// Create a struct member write node to store the parameters into the animation node
-	UK2Node_StructMemberSet* AssignmentNode = SpawnIntermediateNode<UK2Node_StructMemberSet>(VisualAnimNode, ConsolidatedEventGraph);
-	AssignmentNode->VariableReference.SetSelfMember(Record.NodeVariableProperty->GetFName());
-	AssignmentNode->StructType = Record.NodeVariableProperty->Struct;
-	AssignmentNode->AllocateDefaultPins();
+			// Should be true, we only want to deal with instance targets in here
+			if (PropHandler.bInstanceIsTarget)
+			{
+				for (FPropertyCopyRecord& CopyRecord : PropHandler.CopyRecords)
+				{
+					// New set node for the property
+					UK2Node_VariableSet* VarAssignNode = SpawnIntermediateNode<UK2Node_VariableSet>(VisualAnimNode, ConsolidatedEventGraph);
+					VarAssignNode->VariableReference.SetSelfMember(CopyRecord.DestProperty->GetFName());
+					VarAssignNode->AllocateDefaultPins();
 
-	// Wire up the variable node execution wires
-	UEdGraphPin* ExecVariablesIn = Schema->FindExecutionPin(*AssignmentNode, EGPD_Input);
-	ExecChain->MakeLinkTo(ExecVariablesIn);
-	ExecChain = Schema->FindExecutionPin(*AssignmentNode, EGPD_Output);
+					// Wire up the exec line, and update the end of the chain
+					UEdGraphPin* ExecVariablesIn = Schema->FindExecutionPin(*VarAssignNode, EGPD_Input);
+					ExecChain->MakeLinkTo(ExecVariablesIn);
+					ExecChain = Schema->FindExecutionPin(*VarAssignNode, EGPD_Output);
 
-	// Run thru each property
-	TSet<FName> PropertiesBeingSet;
+					// Find the property pin on the set node and configure
+					for (UEdGraphPin* TargetPin : VarAssignNode->Pins)
+					{
+						if (TargetPin->PinType.IsContainer())
+						{
+							// Currently unsupported
+							continue;
+						}
 
-	for (auto TargetPinIt = AssignmentNode->Pins.CreateIterator(); TargetPinIt; ++TargetPinIt)
-	{
-		UEdGraphPin* TargetPin = *TargetPinIt;
-		FName PropertyName(TargetPin->PinName);
+						FName PinPropertyName(TargetPin->PinName);
 
-		// Does it get serviced by this handler?
-		if (FAnimNodeSinglePropertyHandler* SourceInfo = Record.ServicedProperties.Find(PropertyName))
-		{
-			if (TargetPin->PinType.IsArray())
-			{
-				// Grab the array that we need to set members for
-				UK2Node_StructMemberGet* FetchArrayNode = SpawnIntermediateNode<UK2Node_StructMemberGet>(VisualAnimNode, ConsolidatedEventGraph);
-				FetchArrayNode->VariableReference.SetSelfMember(Record.NodeVariableProperty->GetFName());
-				FetchArrayNode->StructType = Record.NodeVariableProperty->Struct;
-				FetchArrayNode->AllocatePinsForSingleMemberGet(PropertyName);
+						if (PinPropertyName == PropertyName)
+						{
+							// This is us, wire up the variable
+							UEdGraphPin* DestPin = CopyRecord.DestPin;
 
-				UEdGraphPin* ArrayVariableNode = FetchArrayNode->FindPin(PropertyName);
+							// Copy the data (link up to the source nodes)
+							TargetPin->CopyPersistentDataFromOldPin(*DestPin);
+							MessageLog.NotifyIntermediatePinCreation(TargetPin, DestPin);
 
-				if (SourceInfo->CopyRecords.Num() > 0)
-				{
-					// Set each element in the array
-					for (FPropertyCopyRecord& CopyRecord : SourceInfo->CopyRecords)
-					{
-						int32 ArrayIndex = CopyRecord.DestArrayIndex;
-						UEdGraphPin* DestPin = CopyRecord.DestPin;
-
-						// Create an array element set node
-						UK2Node_CallArrayFunction* ArrayNode = SpawnIntermediateNode<UK2Node_CallArrayFunction>(VisualAnimNode, ConsolidatedEventGraph);
-						ArrayNode->FunctionReference.SetExternalMember(GET_FUNCTION_NAME_CHECKED(UKismetArrayLibrary, Array_Set), UKismetArrayLibrary::StaticClass());
-						ArrayNode->AllocateDefaultPins();
-
-						// Connect the execution chain
-						ExecChain->MakeLinkTo(ArrayNode->GetExecPin());
-						ExecChain = ArrayNode->GetThenPin();
-
-						// Connect the input array
-						UEdGraphPin* TargetArrayPin = ArrayNode->FindPinChecked(TEXT("TargetArray"));
-						TargetArrayPin->MakeLinkTo(ArrayVariableNode);
-						ArrayNode->PinConnectionListChanged(TargetArrayPin);
-
-						// Set the array index
-						UEdGraphPin* TargetIndexPin = ArrayNode->FindPinChecked(TEXT("Index"));
-						TargetIndexPin->DefaultValue = FString::FromInt(ArrayIndex);
-
-						// Wire up the data input
-						UEdGraphPin* TargetItemPin = ArrayNode->FindPinChecked(TEXT("Item"));
-						TargetItemPin->CopyPersistentDataFromOldPin(*DestPin);
-						MessageLog.NotifyIntermediatePinCreation(TargetItemPin, DestPin);
-						DestPin->BreakAllPinLinks();
+							// Old pin needs to not be connected now - break all its links
+							DestPin->BreakAllPinLinks();
+
+							break;
+						}
 					}
-				}
-			}
-			else
-			{
-				check(!TargetPin->PinType.IsContainer())
-				// Single property
-				if (SourceInfo->CopyRecords.Num() > 0 && SourceInfo->CopyRecords[0].DestPin != nullptr)
-				{
-					UEdGraphPin* DestPin = SourceInfo->CopyRecords[0].DestPin;
 
-					PropertiesBeingSet.Add(DestPin->PinName);
-					TargetPin->CopyPersistentDataFromOldPin(*DestPin);
-					MessageLog.NotifyIntermediatePinCreation(TargetPin, DestPin);
-					DestPin->BreakAllPinLinks();
+					//VarAssignNode->ReconstructNode();
 				}
 			}
 		}
 	}
 
-	// Remove any unused pins from the assignment node to avoid smashing constant values
-	for (int32 PinIndex = 0; PinIndex < AssignmentNode->ShowPinForProperties.Num(); ++PinIndex)
+	if (Record.bServicesNodeProperties)
 	{
-		FOptionalPinFromProperty& TestProperty = AssignmentNode->ShowPinForProperties[PinIndex];
-		TestProperty.bShowPin = PropertiesBeingSet.Contains(TestProperty.PropertyName);
-	}
-	AssignmentNode->ReconstructNode();
-}
+		// Create a struct member write node to store the parameters into the animation node
+		UK2Node_StructMemberSet* AssignmentNode = SpawnIntermediateNode<UK2Node_StructMemberSet>(VisualAnimNode, ConsolidatedEventGraph);
+		AssignmentNode->VariableReference.SetSelfMember(Record.NodeVariableProperty->GetFName());
+		AssignmentNode->StructType = Record.NodeVariableProperty->Struct;
+		AssignmentNode->AllocateDefaultPins();
 
-void FAnimBlueprintCompilerContext::CreateEvaluationHandlerInstance(UAnimGraphNode_Base* VisualAnimNode, FEvaluationHandlerRecord& Record)
-{
-	// Shouldn't create a handler if there is nothing to work with
-	check(Record.ServicedProperties.Num() > 0);
-	check(Record.NodeVariableProperty != nullptr);
-	check(Record.bServicesInstanceProperties);
-	
-	const UAnimationGraphSchema* AnimGraphDefaultSchema = GetDefault<UAnimationGraphSchema>();
+		// Wire up the variable node execution wires
+		UEdGraphPin* ExecVariablesIn = Schema->FindExecutionPin(*AssignmentNode, EGPD_Input);
+		ExecChain->MakeLinkTo(ExecVariablesIn);
+		ExecChain = Schema->FindExecutionPin(*AssignmentNode, EGPD_Output);
 
-	// Use the node GUID for a stable name across compiles
-	FString FunctionName = FString::Printf(TEXT("%s_%s_%s_%s"), *AnimGraphDefaultSchema->DefaultEvaluationHandlerName.ToString(), *VisualAnimNode->GetOuter()->GetName(), *VisualAnimNode->GetClass()->GetName(), *VisualAnimNode->NodeGuid.ToString());
-	Record.HandlerFunctionName = FName(*FunctionName);
+		// Run thru each property
+		TSet<FName> PropertiesBeingSet;
 
-	// check function name isnt already used (data exists that can contain duplicate GUIDs) and apply a numeric extension until it is unique
-	int32 ExtensionIndex = 0;
-	FName* ExistingName = HandlerFunctionNames.Find(Record.HandlerFunctionName);
-	while(ExistingName != nullptr)
-	{
-		FunctionName = FString::Printf(TEXT("%s_%s_%s_%s_%d"), *AnimGraphDefaultSchema->DefaultEvaluationHandlerName.ToString(), *VisualAnimNode->GetOuter()->GetName(), *VisualAnimNode->GetClass()->GetName(), *VisualAnimNode->NodeGuid.ToString(), ExtensionIndex);
-		Record.HandlerFunctionName = FName(*FunctionName);
-		ExistingName = HandlerFunctionNames.Find(Record.HandlerFunctionName);
-		ExtensionIndex++;
-	}
-
-	HandlerFunctionNames.Add(Record.HandlerFunctionName);
-
-	// Add a custom event in the graph
-	UK2Node_CustomEvent* EntryNode = SpawnIntermediateNode<UK2Node_CustomEvent>(VisualAnimNode, ConsolidatedEventGraph);
-	EntryNode->bInternalEvent = true;
-	EntryNode->CustomFunctionName = Record.HandlerFunctionName;
-	EntryNode->AllocateDefaultPins();
-
-	// The ExecChain is the current exec output pin in the linear chain
-	UEdGraphPin* ExecChain = Schema->FindExecutionPin(*EntryNode, EGPD_Output);
-
-	// Need to create a variable set call for each serviced property in the handler
-	for(TPair<FName, FAnimNodeSinglePropertyHandler>& PropHandlerPair : Record.ServicedProperties)
-	{
-		FAnimNodeSinglePropertyHandler& PropHandler = PropHandlerPair.Value;
-		FName PropertyName = PropHandlerPair.Key;
-
-		// Should be true, we only want to deal with instance targets in here
-		check(PropHandler.bInstanceIsTarget);
-
-		for(FPropertyCopyRecord& CopyRecord : PropHandler.CopyRecords)
+		for (auto TargetPinIt = AssignmentNode->Pins.CreateIterator(); TargetPinIt; ++TargetPinIt)
 		{
-			// New set node for the property
-			UK2Node_VariableSet* VarAssignNode = SpawnIntermediateNode<UK2Node_VariableSet>(VisualAnimNode, ConsolidatedEventGraph);
-			VarAssignNode->VariableReference.SetSelfMember(CopyRecord.DestProperty->GetFName());
-			VarAssignNode->AllocateDefaultPins();
-
-			// Wire up the exec line, and update the end of the chain
-			UEdGraphPin* ExecVariablesIn = Schema->FindExecutionPin(*VarAssignNode, EGPD_Input);
-			ExecChain->MakeLinkTo(ExecVariablesIn);
-			ExecChain = Schema->FindExecutionPin(*VarAssignNode, EGPD_Output);
+			UEdGraphPin* TargetPin = *TargetPinIt;
+			FName PropertyName(TargetPin->PinName);
 
-			// Find the property pin on the set node and configure
-			for(UEdGraphPin* TargetPin : VarAssignNode->Pins)
+			// Does it get serviced by this handler?
+			if (FAnimNodeSinglePropertyHandler* SourceInfo = Record.ServicedProperties.Find(PropertyName))
 			{
-				if(TargetPin->PinType.IsContainer())
+				if (TargetPin->PinType.IsArray())
 				{
-					// Currently unsupported
-					continue;
-				}
+					// Grab the array that we need to set members for
+					UK2Node_StructMemberGet* FetchArrayNode = SpawnIntermediateNode<UK2Node_StructMemberGet>(VisualAnimNode, ConsolidatedEventGraph);
+					FetchArrayNode->VariableReference.SetSelfMember(Record.NodeVariableProperty->GetFName());
+					FetchArrayNode->StructType = Record.NodeVariableProperty->Struct;
+					FetchArrayNode->AllocatePinsForSingleMemberGet(PropertyName);
 
-				FName PinPropertyName(TargetPin->PinName);
+					UEdGraphPin* ArrayVariableNode = FetchArrayNode->FindPin(PropertyName);
 
-				if(PinPropertyName == PropertyName)
+					if (SourceInfo->CopyRecords.Num() > 0)
+					{
+						// Set each element in the array
+						for (FPropertyCopyRecord& CopyRecord : SourceInfo->CopyRecords)
+						{
+							int32 ArrayIndex = CopyRecord.DestArrayIndex;
+							UEdGraphPin* DestPin = CopyRecord.DestPin;
+
+							// Create an array element set node
+							UK2Node_CallArrayFunction* ArrayNode = SpawnIntermediateNode<UK2Node_CallArrayFunction>(VisualAnimNode, ConsolidatedEventGraph);
+							ArrayNode->FunctionReference.SetExternalMember(GET_FUNCTION_NAME_CHECKED(UKismetArrayLibrary, Array_Set), UKismetArrayLibrary::StaticClass());
+							ArrayNode->AllocateDefaultPins();
+
+							// Connect the execution chain
+							ExecChain->MakeLinkTo(ArrayNode->GetExecPin());
+							ExecChain = ArrayNode->GetThenPin();
+
+							// Connect the input array
+							UEdGraphPin* TargetArrayPin = ArrayNode->FindPinChecked(TEXT("TargetArray"));
+							TargetArrayPin->MakeLinkTo(ArrayVariableNode);
+							ArrayNode->PinConnectionListChanged(TargetArrayPin);
+
+							// Set the array index
+							UEdGraphPin* TargetIndexPin = ArrayNode->FindPinChecked(TEXT("Index"));
+							TargetIndexPin->DefaultValue = FString::FromInt(ArrayIndex);
+
+							// Wire up the data input
+							UEdGraphPin* TargetItemPin = ArrayNode->FindPinChecked(TEXT("Item"));
+							TargetItemPin->CopyPersistentDataFromOldPin(*DestPin);
+							MessageLog.NotifyIntermediatePinCreation(TargetItemPin, DestPin);
+							DestPin->BreakAllPinLinks();
+						}
+					}
+				}
+				else
 				{
-					// This is us, wire up the variable
-					UEdGraphPin* DestPin = CopyRecord.DestPin;
-
-					// Copy the data (link up to the source nodes)
-					TargetPin->CopyPersistentDataFromOldPin(*DestPin);
-					MessageLog.NotifyIntermediatePinCreation(TargetPin, DestPin);
-
-					// Old pin needs to not be connected now - break all its links
-					DestPin->BreakAllPinLinks();
+					check(!TargetPin->PinType.IsContainer());
+					// Single property
+					if (SourceInfo->CopyRecords.Num() > 0 && SourceInfo->CopyRecords[0].DestPin != nullptr)
+					{
+						UEdGraphPin* DestPin = SourceInfo->CopyRecords[0].DestPin;
 
-					break;
+						PropertiesBeingSet.Add(DestPin->PinName);
+						TargetPin->CopyPersistentDataFromOldPin(*DestPin);
+						MessageLog.NotifyIntermediatePinCreation(TargetPin, DestPin);
+						DestPin->BreakAllPinLinks();
+					}
 				}
 			}
+		}
 
-			//VarAssignNode->ReconstructNode();
+		// Remove any unused pins from the assignment node to avoid smashing constant values
+		for (int32 PinIndex = 0; PinIndex < AssignmentNode->ShowPinForProperties.Num(); ++PinIndex)
+		{
+			FOptionalPinFromProperty& TestProperty = AssignmentNode->ShowPinForProperties[PinIndex];
+			TestProperty.bShowPin = PropertiesBeingSet.Contains(TestProperty.PropertyName);
 		}
+
+		AssignmentNode->ReconstructNode();
 	}
 }
 
+
 void FAnimBlueprintCompilerContext::ProcessAnimationNode(UAnimGraphNode_Base* VisualAnimNode)
 {
 	// Early out if this node has already been processed
@@ -645,20 +617,20 @@ void FAnimBlueprintCompilerContext::ProcessAnimationNode(UAnimGraphNode_Base* Vi
 			// owning anim instance if using a linked instance node.
 			UProperty* SourcePinProperty = nullptr;
 			int32 SourceArrayIndex = INDEX_NONE;
+			bool bInstancePropertyExists = false;
 
 			// We have special handling below if we're targeting a linked instance instead of our own instance properties
 			UAnimGraphNode_CustomProperty* CustomPropertyNode = Cast<UAnimGraphNode_CustomProperty>(VisualAnimNode);
 
+			VisualAnimNode->GetPinAssociatedProperty(NodeType, SourcePin, /*out*/ SourcePinProperty, /*out*/ SourceArrayIndex);
+
 			// Does this pin have an associated evaluation handler?
-			if(CustomPropertyNode)
+			if(!SourcePinProperty && CustomPropertyNode)
 			{
 				// Custom property nodes use instance properties not node properties as they aren't UObjects
 				// and we can't store non-native properties there
 				CustomPropertyNode->GetInstancePinProperty(NewAnimBlueprintClass, SourcePin, SourcePinProperty);
-			}
-			else
-			{
-				VisualAnimNode->GetPinAssociatedProperty(NodeType, SourcePin, /*out*/ SourcePinProperty, /*out*/ SourceArrayIndex);
+				bInstancePropertyExists = true;
 			}
 			
 			if (SourcePinProperty != NULL)
@@ -687,16 +659,18 @@ void FAnimBlueprintCompilerContext::ProcessAnimationNode(UAnimGraphNode_Base* Vi
 					ensure(EvalHandler.NodeVariableProperty == nullptr || EvalHandler.NodeVariableProperty == NewProperty);
 					EvalHandler.NodeVariableProperty = NewProperty;
 					EvalHandler.RegisterPin(SourcePin, SourcePinProperty, SourceArrayIndex);
+					// if it's not instance property, ensure we mark it
+					EvalHandler.bServicesNodeProperties = EvalHandler.bServicesNodeProperties | !bInstancePropertyExists;
 
-					if(CustomPropertyNode)
+					if (CustomPropertyNode)
 					{
-						EvalHandler.bServicesInstanceProperties = true;
+						EvalHandler.bServicesInstanceProperties = EvalHandler.bServicesInstanceProperties | bInstancePropertyExists;
 
 						FAnimNodeSinglePropertyHandler* SinglePropHandler = EvalHandler.ServicedProperties.Find(SourcePinProperty->GetFName());
 						check(SinglePropHandler); // Should have been added in RegisterPin
 
 						// Flag that the target property is actually on the instance class and not the node
-						SinglePropHandler->bInstanceIsTarget = true;
+						SinglePropHandler->bInstanceIsTarget = bInstancePropertyExists;
 					}
 
 					bConsumed = true;
@@ -740,15 +714,7 @@ void FAnimBlueprintCompilerContext::ProcessAnimationNode(UAnimGraphNode_Base* Vi
 			NewAnimBlueprintClass->EvaluateGraphExposedInputs.AddDefaulted(StructEvalHandlers.Num());
 			Record.EvaluationHandlerIdx = NewAnimBlueprintClass->EvaluateGraphExposedInputs.Num() - 1;
 
-			// add instance to class:
-			if(Record.bServicesInstanceProperties)
-			{
-				CreateEvaluationHandlerInstance(VisualAnimNode, Record);
-			}
-			else
-			{
-				CreateEvaluationHandlerStruct(VisualAnimNode, Record);
-			}
+			CreateEvaluationHandler(VisualAnimNode, Record);
 
 			ValidEvaluationHandlerList.Add(Record);
 		}
diff --git a/Engine/Source/Editor/KismetCompiler/Private/AnimBlueprintCompiler.h b/Engine/Source/Editor/KismetCompiler/Private/AnimBlueprintCompiler.h
--- a/Engine/Source/Editor/KismetCompiler/Private/AnimBlueprintCompiler.h
+++ b/Engine/Source/Editor/KismetCompiler/Private/AnimBlueprintCompiler.h
@@ -161,6 +161,7 @@ protected:
 		bool Apply(UObject* Object);
 	};
 
+	/** BP execution handler for Anim node - possibly */
 	struct FEvaluationHandlerRecord
 	{
 	public:
@@ -171,6 +172,9 @@ protected:
 		// The specific evaluation handler inside the specified node
 		int32 EvaluationHandlerIdx;
 
+		// Whether or not our serviced properties are actually on the anim node 
+		bool bServicesNodeProperties;
+
 		// Whether or not our serviced properties are actually on the instance instead of the node
 		bool bServicesInstanceProperties;
 
@@ -185,6 +189,7 @@ protected:
 		FEvaluationHandlerRecord()
 			: NodeVariableProperty(nullptr)
 			, EvaluationHandlerIdx(INDEX_NONE)
+			, bServicesNodeProperties(false)
 			, bServicesInstanceProperties(false)
 			, HandlerFunctionName(NAME_None)
 		{}
@@ -298,8 +303,7 @@ private:
 	UK2Node_CallFunction* SpawnCallAnimInstanceFunction(UEdGraphNode* SourceNode, FName FunctionName);
 
 	// Creates an evaluation handler for an FExposedValue property in an animation node
-	void CreateEvaluationHandlerStruct(UAnimGraphNode_Base* VisualAnimNode, FEvaluationHandlerRecord& Record);
-	void CreateEvaluationHandlerInstance(UAnimGraphNode_Base* VisualAnimNode, FEvaluationHandlerRecord& Record);
+	void CreateEvaluationHandler(UAnimGraphNode_Base* VisualAnimNode, FEvaluationHandlerRecord& Record);
 
 	// Prunes any nodes that aren't reachable via a pose link
 	void PruneIsolatedAnimationNodes(const TArray<UAnimGraphNode_Base*>& RootSet, TArray<UAnimGraphNode_Base*>& GraphNodes);
diff --git a/Engine/Source/Editor/MovieSceneTools/Private/TrackEditors/TransformTrackEditor.cpp b/Engine/Source/Editor/MovieSceneTools/Private/TrackEditors/TransformTrackEditor.cpp
--- a/Engine/Source/Editor/MovieSceneTools/Private/TrackEditors/TransformTrackEditor.cpp
+++ b/Engine/Source/Editor/MovieSceneTools/Private/TrackEditors/TransformTrackEditor.cpp
@@ -195,12 +195,12 @@ void F3DTransformTrackEditor::BuildTrackContextMenu( FMenuBuilder& MenuBuilder,
 	}
 	UMovieScene3DTransformTrack* TransformTrack = Cast<UMovieScene3DTransformTrack>( Track );
 	MenuBuilder.AddMenuEntry(
-		NSLOCTEXT( "Sequencer", "PasteMatineeMoveTrack", "Paste Matinee Move Track" ),
-		NSLOCTEXT( "Sequencer", "PasteMatineeMoveTrackTooltip", "Pastes keys from a Matinee move track into this track." ),
+		NSLOCTEXT("Sequencer", "PasteMatineeMoveTrack", "Paste Matinee Move Track"),
+		NSLOCTEXT("Sequencer", "PasteMatineeMoveTrackTooltip", "Pastes keys from a Matinee move track into this track."),
 		FSlateIcon(),
 		FUIAction(
-			FExecuteAction::CreateStatic( &CopyInterpMoveTrack, GetSequencer().ToSharedRef(), MoveTrack, TransformTrack ),
-			FCanExecuteAction::CreateStatic( &CanCopyInterpMoveTrack, MoveTrack, TransformTrack ) ) );
+			FExecuteAction::CreateStatic(&CopyInterpMoveTrack, GetSequencer().ToSharedRef(), MoveTrack, TransformTrack),
+			FCanExecuteAction::CreateStatic(&CanCopyInterpMoveTrack, MoveTrack, TransformTrack)));
 
 	//		FCanExecuteAction::CreateLambda( [=]()->bool { return MoveTrack != nullptr && MoveTrack->GetNumKeys() > 0 && TransformTrack != nullptr; } ) ) );
 
@@ -389,6 +389,35 @@ void F3DTransformTrackEditor::OnPostPropertyChanged(UObject* InObject, FProperty
 	}
 }
 
+bool F3DTransformTrackEditor::CanAddTransformKeysForSelectedObjects() const
+{
+	// WASD hotkeys to fly the viewport can conflict with hotkeys for setting keyframes (ie. s). 
+// If the viewport is moving, disregard setting keyframes.
+	for (FLevelEditorViewportClient* LevelVC : GEditor->GetLevelViewportClients())
+	{
+		if (LevelVC && LevelVC->IsMovingCamera())
+		{
+			return false;
+		}
+	}
+	TArray<UObject*> SelectedObjects;
+	for (FSelectedEditableComponentIterator It(GEditor->GetSelectedEditableComponentIterator()); It; ++It)
+	{
+		USceneComponent* SceneComponent = Cast<USceneComponent>(*It);
+		if (SceneComponent)
+		{
+			return true;
+		}
+	}
+
+	if (SelectedObjects.Num() == 0)
+	{
+		USelection* CurrentSelection = GEditor->GetSelectedActors();
+		CurrentSelection->GetSelectedObjects(AActor::StaticClass(), SelectedObjects);
+	}
+	return SelectedObjects.Num() > 0;
+}
+
 void F3DTransformTrackEditor::OnAddTransformKeysForSelectedObjects( EMovieSceneTransformChannel Channel )
 {
 	// WASD hotkeys to fly the viewport can conflict with hotkeys for setting keyframes (ie. s). 
@@ -431,20 +460,22 @@ void F3DTransformTrackEditor::BindCommands(TSharedRef<FUICommandList> SequencerC
 	CommandBindings = MakeShared<FUICommandList>();
 	CommandBindings->MapAction(
 		Commands.AddTransformKey,
-		FExecuteAction::CreateSP( this, &F3DTransformTrackEditor::OnAddTransformKeysForSelectedObjects, EMovieSceneTransformChannel::All ) );
+		FExecuteAction::CreateSP(this, &F3DTransformTrackEditor::OnAddTransformKeysForSelectedObjects, EMovieSceneTransformChannel::All),
+		FCanExecuteAction::CreateSP(this, &F3DTransformTrackEditor::CanAddTransformKeysForSelectedObjects));
+		
 
 	CommandBindings->MapAction(
 		Commands.AddTranslationKey,
-		FExecuteAction::CreateSP( this, &F3DTransformTrackEditor::OnAddTransformKeysForSelectedObjects, EMovieSceneTransformChannel::Translation ) );
-
+		FExecuteAction::CreateSP( this, &F3DTransformTrackEditor::OnAddTransformKeysForSelectedObjects, EMovieSceneTransformChannel::Translation ),
+		FCanExecuteAction::CreateSP(this, &F3DTransformTrackEditor::CanAddTransformKeysForSelectedObjects));
 	CommandBindings->MapAction(
 		Commands.AddRotationKey,
-		FExecuteAction::CreateSP( this, &F3DTransformTrackEditor::OnAddTransformKeysForSelectedObjects, EMovieSceneTransformChannel::Rotation ) );
-
+		FExecuteAction::CreateSP( this, &F3DTransformTrackEditor::OnAddTransformKeysForSelectedObjects, EMovieSceneTransformChannel::Rotation ),
+		FCanExecuteAction::CreateSP(this, &F3DTransformTrackEditor::CanAddTransformKeysForSelectedObjects));
 	CommandBindings->MapAction(
 		Commands.AddScaleKey,
-		FExecuteAction::CreateSP( this, &F3DTransformTrackEditor::OnAddTransformKeysForSelectedObjects, EMovieSceneTransformChannel::Scale ) );
-
+		FExecuteAction::CreateSP( this, &F3DTransformTrackEditor::OnAddTransformKeysForSelectedObjects, EMovieSceneTransformChannel::Scale ),
+		FCanExecuteAction::CreateSP(this, &F3DTransformTrackEditor::CanAddTransformKeysForSelectedObjects));
 	Commands.BindingCount++;
 
 	// Add these bindings to Sequencer
@@ -869,20 +900,20 @@ bool F3DTransformTrackEditor::ModifyGeneratedKeysByCurrentAndWeight(UObject *Obj
 		CurrentPos = Transform.Translation;
 		CurrentRot = Transform.Rotation;
 		CurrentScale = Transform.Scale;
-		break;
-	}
-	FMovieSceneChannelProxy& Proxy = SectionToKey->GetChannelProxy();
-	GeneratedTotalKeys[0]->ModifyByCurrentAndWeight(Proxy, KeyTime, (void *)&CurrentPos.X, Weight);
-	GeneratedTotalKeys[1]->ModifyByCurrentAndWeight(Proxy, KeyTime, (void *)&CurrentPos.Y, Weight);
-	GeneratedTotalKeys[2]->ModifyByCurrentAndWeight(Proxy, KeyTime, (void *)&CurrentPos.Z, Weight);
-	GeneratedTotalKeys[3]->ModifyByCurrentAndWeight(Proxy, KeyTime, (void *)&CurrentRot.Roll, Weight);
-	GeneratedTotalKeys[4]->ModifyByCurrentAndWeight(Proxy, KeyTime, (void *)&CurrentRot.Pitch, Weight);
-	GeneratedTotalKeys[5]->ModifyByCurrentAndWeight(Proxy, KeyTime, (void *)&CurrentRot.Yaw, Weight);
-	GeneratedTotalKeys[6]->ModifyByCurrentAndWeight(Proxy, KeyTime, (void *)&CurrentScale.X, Weight);
-	GeneratedTotalKeys[7]->ModifyByCurrentAndWeight(Proxy, KeyTime, (void *)&CurrentScale.Y, Weight);
-	GeneratedTotalKeys[8]->ModifyByCurrentAndWeight(Proxy, KeyTime, (void *)&CurrentScale.Z, Weight);
-	return true;
 
+		FMovieSceneChannelProxy& Proxy = SectionToKey->GetChannelProxy();
+		GeneratedTotalKeys[0]->ModifyByCurrentAndWeight(Proxy, KeyTime, (void *)&CurrentPos.X, Weight);
+		GeneratedTotalKeys[1]->ModifyByCurrentAndWeight(Proxy, KeyTime, (void *)&CurrentPos.Y, Weight);
+		GeneratedTotalKeys[2]->ModifyByCurrentAndWeight(Proxy, KeyTime, (void *)&CurrentPos.Z, Weight);
+		GeneratedTotalKeys[3]->ModifyByCurrentAndWeight(Proxy, KeyTime, (void *)&CurrentRot.Roll, Weight);
+		GeneratedTotalKeys[4]->ModifyByCurrentAndWeight(Proxy, KeyTime, (void *)&CurrentRot.Pitch, Weight);
+		GeneratedTotalKeys[5]->ModifyByCurrentAndWeight(Proxy, KeyTime, (void *)&CurrentRot.Yaw, Weight);
+		GeneratedTotalKeys[6]->ModifyByCurrentAndWeight(Proxy, KeyTime, (void *)&CurrentScale.X, Weight);
+		GeneratedTotalKeys[7]->ModifyByCurrentAndWeight(Proxy, KeyTime, (void *)&CurrentScale.Y, Weight);
+		GeneratedTotalKeys[8]->ModifyByCurrentAndWeight(Proxy, KeyTime, (void *)&CurrentScale.Z, Weight);
+		return true;
+	}
+	return false;
 }
 
 void AddUnwoundKey(FMovieSceneFloatChannel& Channel, FFrameNumber Time, float Value)
@@ -957,6 +988,10 @@ void F3DTransformTrackEditor::ImportAnimSequenceTransforms(const FAssetData& Ass
 		TransformTrack->Modify();
 
 		UMovieScene3DTransformSection* Section = Cast<UMovieScene3DTransformSection>(TransformTrack->CreateNewSection());
+
+		Section->SetBlendType(EMovieSceneBlendType::Additive);
+		Section->SetMask(EMovieSceneTransformChannel::Translation | EMovieSceneTransformChannel::Rotation);
+		
 		FFrameRate TickResolution = Section->GetTypedOuter<UMovieScene>()->GetTickResolution();
 
 		TArrayView<FMovieSceneFloatChannel*> FloatChannels = Section->GetChannelProxy().GetChannels<FMovieSceneFloatChannel>();
diff --git a/Engine/Source/Editor/MovieSceneTools/Private/TrackEditors/TransformTrackEditor.h b/Engine/Source/Editor/MovieSceneTools/Private/TrackEditors/TransformTrackEditor.h
--- a/Engine/Source/Editor/MovieSceneTools/Private/TrackEditors/TransformTrackEditor.h
+++ b/Engine/Source/Editor/MovieSceneTools/Private/TrackEditors/TransformTrackEditor.h
@@ -113,6 +113,12 @@ private:
 	/** Delegate for camera button lock tooltip */
 	FText GetLockCameraToolTip(FGuid ObjectGuid) const; 
 
+
+	/** Whether or not we can add a transform key for a selected object 
+	* @return Returns true if we can.
+	**/
+	bool CanAddTransformKeysForSelectedObjects() const;
+
 	/** Generates transform keys based on the last transform, the current transform, and other options. 
 		One transform key is generated for each individual key to be added to the section. */
 	void GetTransformKeys( const TOptional<FTransformData>& LastTransform, const FTransformData& CurrentTransform, EMovieSceneTransformChannel ChannelsToKey, FGeneratedTrackKeys& OutGeneratedKeys );
diff --git a/Engine/Source/Editor/MovieSceneTools/Public/KeyframeTrackEditor.h b/Engine/Source/Editor/MovieSceneTools/Public/KeyframeTrackEditor.h
--- a/Engine/Source/Editor/MovieSceneTools/Public/KeyframeTrackEditor.h
+++ b/Engine/Source/Editor/MovieSceneTools/Public/KeyframeTrackEditor.h
@@ -142,7 +142,7 @@ struct TAddKeyImpl<FMovieSceneFloatChannel, float> : IImpl
 		using namespace MovieScene;
 
 		FMovieSceneFloatChannel* Channel = Proxy.GetChannel<FMovieSceneFloatChannel>(ChannelIndex);
-		if (Channel && Channel->GetData().GetTimes().Num() == 0)
+		if (Channel && Channel->GetData().GetTimes().Num() == 0  && Channel->GetDefault() != ValueToSet)
 		{
 			if (Section->TryModify())
 			{
@@ -499,6 +499,7 @@ private:
 		return KeyPropertyResult;
 	}
 
+protected:
 	/* Returns whether a section was added */
 	FKeyPropertyResult AddKeysToSection(UMovieSceneSection* Section, FFrameNumber KeyTime, const FGeneratedTrackKeys& Keys, ESequencerKeyMode KeyMode)
 	{
diff --git a/Engine/Source/Editor/Persona/Persona.Build.cs b/Engine/Source/Editor/Persona/Persona.Build.cs
--- a/Engine/Source/Editor/Persona/Persona.Build.cs
+++ b/Engine/Source/Editor/Persona/Persona.Build.cs
@@ -30,6 +30,7 @@ public class Persona : ModuleRules
                 "AnimationEditor",
                 "MeshReductionInterface",
                 "SequenceRecorder",
+                "AnimationBlueprintEditor",
             }
         );
 
@@ -55,9 +56,9 @@ public class Persona : ModuleRules
                 "RHI",
                 "Json",
                 "JsonUtilities",
-                "ClothingSystemRuntimeNv",
                 "ClothingSystemEditorInterface",
                 "ClothingSystemRuntimeInterface",
+                "ClothingSystemRuntimeCommon",
                 "AnimGraphRuntime",
                 "UnrealEd",
                 "CommonMenuExtensions",
