Add Persona as an include dependency to ClothingSystemEditor.

diff --git a/Engine/Source/Editor/ClothingSystemEditor/ClothingSystemEditor.Build.cs b/Engine/Source/Editor/ClothingSystemEditor/ClothingSystemEditor.Build.cs
--- a/Engine/Source/Editor/ClothingSystemEditor/ClothingSystemEditor.Build.cs
+++ b/Engine/Source/Editor/ClothingSystemEditor/ClothingSystemEditor.Build.cs
@@ -11,8 +11,9 @@ public class ClothingSystemEditor : ModuleRules
         PublicIncludePathModuleNames.Add("ClothingSystemRuntimeInterface");
         PublicIncludePathModuleNames.Add("ClothingSystemRuntimeNv");
         PublicIncludePathModuleNames.Add("ClothingSystemEditorInterface");
+		PublicIncludePathModuleNames.Add("Persona");
 
-        PrivateDependencyModuleNames.AddRange(
+		PrivateDependencyModuleNames.AddRange(
             new string[]
             {
                 "Engine",
