Cloth - Add cloth vertex budget warning in editor.

diff --git a/Engine/Source/Editor/ClothingSystemEditor/Private/ClothingAssetFactory.cpp b/Engine/Source/Editor/ClothingSystemEditor/Private/ClothingAssetFactory.cpp
--- a/Engine/Source/Editor/ClothingSystemEditor/Private/ClothingAssetFactory.cpp
+++ b/Engine/Source/Editor/ClothingSystemEditor/Private/ClothingAssetFactory.cpp
@@ -1306,6 +1306,16 @@ bool UClothingAssetFactory::ImportToLodInternal(
 		DestAsset->ApplyParameterMasks();
 	}
 
+	int32 LODVertexBudget;
+	if (GConfig->GetInt(TEXT("ClothSettings"), TEXT("LODVertexBudget"), LODVertexBudget, GEditorIni) && LODVertexBudget > 0 && NumUniqueVerts > LODVertexBudget)
+	{
+		LogAndToastWarning(FText::Format(LOCTEXT("LODVertexBudgetWarning", 
+			"This cloth LOD has {0} more vertices than what is budgeted on this project (current={1}, budget={2})"), 
+			NumUniqueVerts - LODVertexBudget, 
+			NumUniqueVerts,
+			LODVertexBudget));
+	}
+
 	return true;
 }
 
