CIS - Fix non-unity build error caused by CL 12976746.

diff --git a/Engine/Source/Editor/ClothingSystemEditor/Private/ClothingAssetFactory.cpp b/Engine/Source/Editor/ClothingSystemEditor/Private/ClothingAssetFactory.cpp
--- a/Engine/Source/Editor/ClothingSystemEditor/Private/ClothingAssetFactory.cpp
+++ b/Engine/Source/Editor/ClothingSystemEditor/Private/ClothingAssetFactory.cpp
@@ -25,6 +25,7 @@
 #include "Utils/ClothingMeshUtils.h"
 #include "Rendering/SkeletalMeshModel.h"
 #include "ClothPhysicalMeshData.h"
+#include "Misc/ConfigCacheIni.h"
 
 #define LOCTEXT_NAMESPACE "ClothingAssetFactory"
 DEFINE_LOG_CATEGORY(LogClothingAssetFactory)
