Chaos Cloth - Fix MaxDistance weightmap creation after import.

diff --git a/Engine/Source/Editor/ClothingSystemEditor/Private/ClothingAssetFactory.cpp b/Engine/Source/Editor/ClothingSystemEditor/Private/ClothingAssetFactory.cpp
--- a/Engine/Source/Editor/ClothingSystemEditor/Private/ClothingAssetFactory.cpp
+++ b/Engine/Source/Editor/ClothingSystemEditor/Private/ClothingAssetFactory.cpp
@@ -1252,11 +1252,14 @@ bool UClothingAssetFactory::ImportToLodInternal(
 		}
 	}
 
-	// Add a max distance parameter mask to begin with
+	// Add a max distance parameter mask to the physics mesh
+	FPointWeightMap& PhysMeshMaxDistances = PhysMesh.AddWeightMap(EWeightMapTargetCommon::MaxDistance);
+	PhysMeshMaxDistances.Initialize(PhysMesh.Vertices.Num());
+
+	// Add a max distance parameter mask to the LOD
 	DestLod.PointWeightMaps.AddDefaulted();
-	FPointWeightMap& Mask = DestLod.PointWeightMaps.Last();
-	const FPointWeightMap& MaxDistances = PhysMesh.GetWeightMap(EWeightMapTargetCommon::MaxDistance);
-	Mask.Initialize(MaxDistances, EWeightMapTargetCommon::MaxDistance);
+	FPointWeightMap& LodMaxDistances = DestLod.PointWeightMaps.Last();
+	LodMaxDistances.Initialize(PhysMeshMaxDistances, EWeightMapTargetCommon::MaxDistance);
 
 	PhysMesh.MaxBoneWeights = SourceSection.MaxBoneInfluences;
 
