Physics interface cleanup.

* Removed deprecated or dead code paths
* Simplified build system setup for physics support
* Deprecated build system flags and unsupported macros

diff --git a/Engine/Platforms/Hololens/Source/Programs/UnrealBuildTool/UEBuildHoloLens.cs b/Engine/Platforms/Hololens/Source/Programs/UnrealBuildTool/UEBuildHoloLens.cs
--- a/Engine/Platforms/Hololens/Source/Programs/UnrealBuildTool/UEBuildHoloLens.cs
+++ b/Engine/Platforms/Hololens/Source/Programs/UnrealBuildTool/UEBuildHoloLens.cs
@@ -273,7 +273,6 @@ namespace UnrealBuildTool
 			Target.WindowsPlatform.bStrictConformanceMode = false;
 
 			Target.bDeployAfterCompile = true;
-			Target.bCompileNvCloth = false;      // requires CUDA
 
 			// Disable Simplygon support if compiling against the NULL RHI.
 			if (Target.GlobalDefinitions.Contains("USE_NULL_RHI=1"))
diff --git a/Engine/Plugins/Experimental/ChaosCaching/Source/ChaosCaching/Private/Chaos/Adapters/GeometryCollectionComponentCacheAdapter.cpp b/Engine/Plugins/Experimental/ChaosCaching/Source/ChaosCaching/Private/Chaos/Adapters/GeometryCollectionComponentCacheAdapter.cpp
--- a/Engine/Plugins/Experimental/ChaosCaching/Source/ChaosCaching/Private/Chaos/Adapters/GeometryCollectionComponentCacheAdapter.cpp
+++ b/Engine/Plugins/Experimental/ChaosCaching/Source/ChaosCaching/Private/Chaos/Adapters/GeometryCollectionComponentCacheAdapter.cpp
@@ -588,8 +588,6 @@ namespace Chaos
 
 	Chaos::FPhysicsSolver* FGeometryCollectionCacheAdapter::GetComponentSolver(UPrimitiveComponent* InComponent) const
 	{
-#if WITH_CHAOS
-
 		// If the observed component is a Geometry Collection using a non-default Chaos solver..
 		if (UGeometryCollectionComponent* GeometryCollectionComponent = Cast<UGeometryCollectionComponent>(InComponent))
 		{
@@ -610,7 +608,6 @@ namespace Chaos
 			}
 		}
 
-#endif // WITH_CHAOS
 
 		return nullptr;
 	}
diff --git a/Engine/Plugins/Experimental/ChaosCaching/Source/ChaosCaching/Private/Chaos/Adapters/StaticMeshComponentCacheAdapter.cpp b/Engine/Plugins/Experimental/ChaosCaching/Source/ChaosCaching/Private/Chaos/Adapters/StaticMeshComponentCacheAdapter.cpp
--- a/Engine/Plugins/Experimental/ChaosCaching/Source/ChaosCaching/Private/Chaos/Adapters/StaticMeshComponentCacheAdapter.cpp
+++ b/Engine/Plugins/Experimental/ChaosCaching/Source/ChaosCaching/Private/Chaos/Adapters/StaticMeshComponentCacheAdapter.cpp
@@ -59,14 +59,11 @@ namespace Chaos
 
 	void FStaticMeshCacheAdapter::Record_PostSolve(UPrimitiveComponent* InComponent, const FTransform& InRootTransform, FPendingFrameWrite& OutFrame, Chaos::FReal InTime) const
 	{
-#if WITH_CHAOS
-
 		UStaticMeshComponent* MeshComp = CastChecked<UStaticMeshComponent>(InComponent);
 
 		FSingleParticlePhysicsProxy* PhysProxy = MeshComp->BodyInstance.ActorHandle;
 
 		RecordToCacheInternal(PhysProxy, InRootTransform, OutFrame, InTime);
-#endif // WITH_CHAOS
 	}
 
 	void PlayFromCacheInternal(FSingleParticlePhysicsProxy* InProxy, UChaosCache* InCache, FPlaybackTickRecord& TickRecord, TArray<TPBDRigidParticleHandle<Chaos::FReal, 3>*>& OutUpdatedRigids) 
@@ -103,14 +100,11 @@ namespace Chaos
 
 	void FStaticMeshCacheAdapter::Playback_PreSolve(UPrimitiveComponent* InComponent, UChaosCache* InCache, Chaos::FReal InTime, FPlaybackTickRecord& TickRecord, TArray<TPBDRigidParticleHandle<Chaos::FReal, 3>*>& OutUpdatedRigids) const
 	{
-#if WITH_CHAOS
-
 		UStaticMeshComponent* MeshComp = CastChecked<UStaticMeshComponent>(InComponent);
 
 		FSingleParticlePhysicsProxy* PhysProxy = MeshComp->BodyInstance.ActorHandle;
 
 		PlayFromCacheInternal(PhysProxy, InCache, TickRecord, OutUpdatedRigids);
-#endif // WITH_CHAOS
 	}
 
 	FGuid FStaticMeshCacheAdapter::GetGuid() const
@@ -129,8 +123,6 @@ namespace Chaos
 
 	Chaos::FPhysicsSolver* FStaticMeshCacheAdapter::GetComponentSolver(UPrimitiveComponent* InComponent) const
 	{
-#if WITH_CHAOS
-
 		if(InComponent && InComponent->GetWorld())
 		{
 			UWorld* ComponentWorld = InComponent->GetWorld();
@@ -140,7 +132,6 @@ namespace Chaos
 				return WorldScene->GetSolver();
 			}
 		}
-#endif // WITH_CHAOS
 
 		return nullptr;
 	}
@@ -179,14 +170,10 @@ namespace Chaos
 	{
 		EnsureIsInGameThreadContext();
 
-
-#if WITH_CHAOS
-
 		if(Cast<UStaticMeshComponent>(InComponent) && InComponent->GetBodyInstance() && InComponent->GetBodyInstance()->ActorHandle)
 		{
 			FPhysInterface_Chaos::SetIsKinematic_AssumesLocked(InComponent->GetBodyInstance()->ActorHandle, true);
 		}
-#endif // WITH_CHAOS
 
 		return true;
 	}
diff --git a/Engine/Plugins/Experimental/ChaosCloth/Source/Chaos/Private/ChaosCloth/ChaosClothModule.cpp b/Engine/Plugins/Experimental/ChaosCloth/Source/Chaos/Private/ChaosCloth/ChaosClothModule.cpp
--- a/Engine/Plugins/Experimental/ChaosCloth/Source/Chaos/Private/ChaosCloth/ChaosClothModule.cpp
+++ b/Engine/Plugins/Experimental/ChaosCloth/Source/Chaos/Private/ChaosCloth/ChaosClothModule.cpp
@@ -15,24 +15,17 @@ class FChaosClothModule : public IChaosClothModuleInterface, public IClothingSim
     virtual void StartupModule() override
     {
         check(GConfig);
-#if WITH_CHAOS
 		IModularFeatures::Get().RegisterModularFeature(IClothingSimulationFactoryClassProvider::FeatureName, this);
-#endif
     }
 
     virtual void ShutdownModule() override
     {
-#if WITH_CHAOS
 		IModularFeatures::Get().UnregisterModularFeature(IClothingSimulationFactoryClassProvider::FeatureName, this);
-#endif
     }
 
 	TSubclassOf<UClothingSimulationFactory> GetClothingSimulationFactoryClass() const override
 	{
-#if WITH_CHAOS
 		return UChaosClothingSimulationFactory::StaticClass();
-#endif
-		return nullptr;
 	}
 };
 
diff --git a/Engine/Plugins/Experimental/ChaosCloth/Source/Chaos/Private/ChaosCloth/ChaosClothingSimulationCollider.cpp b/Engine/Plugins/Experimental/ChaosCloth/Source/Chaos/Private/ChaosCloth/ChaosClothingSimulationCollider.cpp
--- a/Engine/Plugins/Experimental/ChaosCloth/Source/Chaos/Private/ChaosCloth/ChaosClothingSimulationCollider.cpp
+++ b/Engine/Plugins/Experimental/ChaosCloth/Source/Chaos/Private/ChaosCloth/ChaosClothingSimulationCollider.cpp
@@ -20,13 +20,6 @@
 #include "HAL/PlatformMath.h"
 #include "Containers/ArrayView.h"
 #include "Containers/BitArray.h"
-#if PHYSICS_INTERFACE_PHYSX
-#include "PhysXIncludes.h"
-#endif
-
-#if PHYSICS_INTERFACE_PHYSX
-#include "PhysXPublic.h"
-#endif
 
 DECLARE_CYCLE_STAT(TEXT("Chaos Cloth Update Collider"), STAT_ChaosClothingSimulationColliderUpdate, STATGROUP_ChaosCloth);
 
@@ -529,7 +522,6 @@ void FClothingSimulationCollider::ExtractPhysicsAssetCollision(FClothCollisionDa
 				}
 			}
 
-#if WITH_CHAOS  // Only the Chaos build has access to the ChaosConvex property of the FKConvexElem structure
 			// Add convexes
 			for (const FKConvexElem& ConvexElem : AggGeom.ConvexElems)
 			{
@@ -570,7 +562,6 @@ void FClothingSimulationCollider::ExtractPhysicsAssetCollision(FClothCollisionDa
 				// Add extracted collision data
 				ClothCollisionData.Convexes.Add(Convex);
 			}
-#endif  // #if WITH_CHAOS
 
 			for (const FKLevelSetElem& LevelSetElem : AggGeom.LevelSetElems)
 			{
diff --git a/Engine/Plugins/Experimental/ChaosCloth/Source/Chaos/Private/ChaosCloth/ChaosClothingSimulationFactory.cpp b/Engine/Plugins/Experimental/ChaosCloth/Source/Chaos/Private/ChaosCloth/ChaosClothingSimulationFactory.cpp
--- a/Engine/Plugins/Experimental/ChaosCloth/Source/Chaos/Private/ChaosCloth/ChaosClothingSimulationFactory.cpp
+++ b/Engine/Plugins/Experimental/ChaosCloth/Source/Chaos/Private/ChaosCloth/ChaosClothingSimulationFactory.cpp
@@ -19,11 +19,7 @@ void UChaosClothingSimulationFactory::DestroySimulation(IClothingSimulation* InS
 
 bool UChaosClothingSimulationFactory::SupportsAsset(UClothingAssetBase* InAsset)
 {
-#if WITH_CHAOS
     return true;
-#else
-    return false;
-#endif
 }
 
 bool UChaosClothingSimulationFactory::SupportsRuntimeInteraction()
diff --git a/Engine/Plugins/Experimental/ChaosClothEditor/Source/Chaos/Private/ChaosClothEditor/ChaosClothEditorModule.cpp b/Engine/Plugins/Experimental/ChaosClothEditor/Source/Chaos/Private/ChaosClothEditor/ChaosClothEditorModule.cpp
--- a/Engine/Plugins/Experimental/ChaosClothEditor/Source/Chaos/Private/ChaosClothEditor/ChaosClothEditorModule.cpp
+++ b/Engine/Plugins/Experimental/ChaosClothEditor/Source/Chaos/Private/ChaosClothEditor/ChaosClothEditorModule.cpp
@@ -14,7 +14,6 @@ DEFINE_LOG_CATEGORY(LogChaosClothEditor);
 
 void FChaosClothEditorModule::StartupModule()
 {
-#if WITH_CHAOS
 	IModularFeatures::Get().RegisterModularFeature(FClothingSystemEditorInterfaceModule::ExtenderFeatureName, &ChaosEditorExtender);
 
 	// Register type customizations
@@ -24,12 +23,10 @@ void FChaosClothEditorModule::StartupModule()
 		// Register weight map property type.
 		PropertyModule->RegisterCustomPropertyTypeLayout("ChaosClothWeightedValue", FOnGetPropertyTypeCustomizationInstance::CreateStatic(&FChaosClothWeightedValueCustomization::MakeInstance));
 	}
-#endif
 }
 
 void FChaosClothEditorModule::ShutdownModule()
 {
-#if WITH_CHAOS
 	IModularFeatures::Get().RegisterModularFeature(FClothingSystemEditorInterfaceModule::ExtenderFeatureName, &ChaosEditorExtender);
 
 	// Unregister type customizations
@@ -38,5 +35,4 @@ void FChaosClothEditorModule::ShutdownModule()
 	{
 		PropertyModule->UnregisterCustomPropertyTypeLayout("ChaosClothWeightedValue");
 	}
-#endif
 }
diff --git a/Engine/Plugins/Experimental/ChaosClothEditor/Source/Chaos/Private/ChaosClothEditor/ChaosSimulationEditorExtender.cpp b/Engine/Plugins/Experimental/ChaosClothEditor/Source/Chaos/Private/ChaosClothEditor/ChaosSimulationEditorExtender.cpp
--- a/Engine/Plugins/Experimental/ChaosClothEditor/Source/Chaos/Private/ChaosClothEditor/ChaosSimulationEditorExtender.cpp
+++ b/Engine/Plugins/Experimental/ChaosClothEditor/Source/Chaos/Private/ChaosClothEditor/ChaosSimulationEditorExtender.cpp
@@ -1,5 +1,4 @@
 // Copyright Epic Games, Inc. All Rights Reserved.
-#if WITH_CHAOS
 
 #include "ChaosClothEditor/ChaosSimulationEditorExtender.h"
 
@@ -214,5 +213,3 @@ void FSimulationEditorExtender::ShowClothSections(USkeletalMeshComponent* MeshCo
 }  // End namespace Chaos
 
 #undef LOCTEXT_NAMESPACE
-
-#endif  // #if WITH_CHAOS
\ No newline at end of file
diff --git a/Engine/Plugins/Experimental/ChaosClothEditor/Source/Chaos/Public/ChaosClothEditor/ChaosClothEditorModule.h b/Engine/Plugins/Experimental/ChaosClothEditor/Source/Chaos/Public/ChaosClothEditor/ChaosClothEditorModule.h
--- a/Engine/Plugins/Experimental/ChaosClothEditor/Source/Chaos/Public/ChaosClothEditor/ChaosClothEditorModule.h
+++ b/Engine/Plugins/Experimental/ChaosClothEditor/Source/Chaos/Public/ChaosClothEditor/ChaosClothEditorModule.h
@@ -13,7 +13,5 @@ public:
 	virtual void ShutdownModule() override;
 
 private:
-#if WITH_CHAOS
 	Chaos::FSimulationEditorExtender ChaosEditorExtender;
-#endif
 };
diff --git a/Engine/Plugins/Experimental/ChaosClothEditor/Source/Chaos/Public/ChaosClothEditor/ChaosSimulationEditorExtender.h b/Engine/Plugins/Experimental/ChaosClothEditor/Source/Chaos/Public/ChaosClothEditor/ChaosSimulationEditorExtender.h
--- a/Engine/Plugins/Experimental/ChaosClothEditor/Source/Chaos/Public/ChaosClothEditor/ChaosSimulationEditorExtender.h
+++ b/Engine/Plugins/Experimental/ChaosClothEditor/Source/Chaos/Public/ChaosClothEditor/ChaosSimulationEditorExtender.h
@@ -2,8 +2,6 @@
 
 #pragma once
 
-#if WITH_CHAOS
-
 #include "SimulationEditorExtender.h"
 #include "Containers/BitArray.h"
 
@@ -33,5 +31,3 @@ namespace Chaos
 		TBitArray<> Flags;
 	};
 }
-
-#endif
diff --git a/Engine/Plugins/Experimental/ChaosEditor/Source/FractureEditor/Private/FractureEditor.cpp b/Engine/Plugins/Experimental/ChaosEditor/Source/FractureEditor/Private/FractureEditor.cpp
--- a/Engine/Plugins/Experimental/ChaosEditor/Source/FractureEditor/Private/FractureEditor.cpp
+++ b/Engine/Plugins/Experimental/ChaosEditor/Source/FractureEditor/Private/FractureEditor.cpp
@@ -10,9 +10,6 @@
 
 void FFractureEditorModule::StartupModule()
 {
-
-#if INCLUDE_CHAOS
-
 	FFractureEditorStyle::Get();
 
 	// This code will execute after your module is loaded into memory; the exact timing is specified in the .uplugin file per-module
@@ -22,8 +19,6 @@ void FFractureEditorModule::StartupModule()
 
 	FFractureEditorCommands::Register();
 
-#endif
-
 }
 
 void FFractureEditorModule::ShutdownModule()
diff --git a/Engine/Plugins/Experimental/ChaosNiagara/Source/ChaosNiagara/Classes/NiagaraDataInterfaceChaosDestruction.h b/Engine/Plugins/Experimental/ChaosNiagara/Source/ChaosNiagara/Classes/NiagaraDataInterfaceChaosDestruction.h
--- a/Engine/Plugins/Experimental/ChaosNiagara/Source/ChaosNiagara/Classes/NiagaraDataInterfaceChaosDestruction.h
+++ b/Engine/Plugins/Experimental/ChaosNiagara/Source/ChaosNiagara/Classes/NiagaraDataInterfaceChaosDestruction.h
@@ -686,9 +686,8 @@ public:
 
 	virtual void ProvidePerInstanceDataForRenderThread(void* DataForRenderThread, void* PerInstanceData, const FNiagaraSystemInstanceID& SystemInstance) override;
 protected:
-#if INCLUDE_CHAOS
+
 	void RegisterWithSolverEventManager(Chaos::FPhysicsSolver* Solver);
-#endif
 	virtual bool CopyToInternal(UNiagaraDataInterface* Destination) const override;
 
 	void ResetInstData(FNDIChaosDestruction_InstanceData* InstData);
diff --git a/Engine/Plugins/Experimental/ChaosVehiclesPlugin/Source/ChaosVehicles/Private/ChaosVehicleManager.cpp b/Engine/Plugins/Experimental/ChaosVehiclesPlugin/Source/ChaosVehicles/Private/ChaosVehicleManager.cpp
--- a/Engine/Plugins/Experimental/ChaosVehiclesPlugin/Source/ChaosVehicles/Private/ChaosVehicleManager.cpp
+++ b/Engine/Plugins/Experimental/ChaosVehiclesPlugin/Source/ChaosVehicles/Private/ChaosVehicleManager.cpp
@@ -27,9 +27,7 @@ uint32 FChaosVehicleManager::VehicleSetupTag = 0;
 FDelegateHandle FChaosVehicleManager::OnPostWorldInitializationHandle;
 FDelegateHandle FChaosVehicleManager::OnWorldCleanupHandle;
 
-#if WITH_CHAOS
 bool FChaosVehicleManager::GInitialized = false;
-#endif
 
 void OnPostWorldInitialize(UWorld* InWorld, const UWorld::InitializationValues)
 {
@@ -50,16 +48,12 @@ void OnWorldCleanup(UWorld* InWorld, bool bSessionEnded, bool bCleanupResources)
 }
 
 FChaosVehicleManager::FChaosVehicleManager(FPhysScene* PhysScene)
-#if WITH_CHAOS
 	: Scene(*PhysScene)
 	, AsyncCallback(nullptr)
 	, Timestamp(0)
-
-#endif
 {
 	check(PhysScene);
-	
-#if WITH_CHAOS
+
 	if (!GInitialized)
 	{
 		GInitialized = true;
@@ -74,24 +68,20 @@ FChaosVehicleManager::FChaosVehicleManager(FPhysScene* PhysScene)
 
 	// Add to Scene-To-Manager map
 	FChaosVehicleManager::SceneToVehicleManagerMap.Add(PhysScene, this);
-#endif
 }
 
 
 void FChaosVehicleManager::RegisterCallbacks()
 {
-#if WITH_CHAOS
 	OnPhysScenePreTickHandle = Scene.OnPhysScenePreTick.AddRaw(this, &FChaosVehicleManager::Update);
 	OnPhysScenePostTickHandle = Scene.OnPhysScenePostTick.AddRaw(this, &FChaosVehicleManager::PostUpdate);
 
 	check(AsyncCallback == nullptr);
 	AsyncCallback = Scene.GetSolver()->CreateAndRegisterSimCallbackObject_External<FChaosVehicleManagerAsyncCallback>(true);
-#endif
 }
 
 void FChaosVehicleManager::UnregisterCallbacks()
 {
-#if WITH_CHAOS
 	Scene.OnPhysScenePreTick.Remove(OnPhysScenePreTickHandle);
 	Scene.OnPhysScenePostTick.Remove(OnPhysScenePostTickHandle);
 
@@ -100,7 +90,6 @@ void FChaosVehicleManager::UnregisterCallbacks()
 		Scene.GetSolver()->UnregisterAndFreeSimCallbackObject_External(AsyncCallback);
 		AsyncCallback = nullptr;
 	}
-#endif
 }
 
 void FChaosVehicleManager::DetachFromPhysScene(FPhysScene* PhysScene)
@@ -174,7 +163,6 @@ void FChaosVehicleManager::ScenePreTick(FPhysScene* PhysScene, float DeltaTime)
 
 void FChaosVehicleManager::Update(FPhysScene* PhysScene, float DeltaTime)
 {
-#if WITH_CHAOS
 	SCOPE_CYCLE_COUNTER(STAT_ChaosVehicleManager_Update);
 
 	UWorld* World = Scene.GetOwningWorld();
@@ -194,7 +182,6 @@ void FChaosVehicleManager::Update(FPhysScene* PhysScene, float DeltaTime)
 			Vehicle->FinalizeSimCallbackData(*AsyncInput);
 		}
 	}
-#endif
 }
 
 void FChaosVehicleManager::PostUpdate(FChaosScene* PhysScene)
@@ -216,7 +203,6 @@ void FChaosVehicleManager::PostUpdate(FChaosScene* PhysScene)
 
 void FChaosVehicleManager::ParallelUpdateVehicles(float DeltaSeconds)
 {
-#if WITH_CHAOS
 	SCOPE_CYCLE_COUNTER(STAT_ChaosVehicleManager_ParallelUpdateVehicles);
 
 	FChaosVehicleManagerAsyncInput* AsyncInput = AsyncCallback->GetProducerInputData_External();
@@ -324,5 +310,4 @@ void FChaosVehicleManager::ParallelUpdateVehicles(float DeltaSeconds)
 
 	bool ForceSingleThread = !GVehicleDebugParams.EnableMultithreading;
 	ParallelFor(AwakeVehiclesBatch.Num(), LambdaParallelUpdate, ForceSingleThread);
-#endif
 }
diff --git a/Engine/Plugins/Experimental/ChaosVehiclesPlugin/Source/ChaosVehicles/Private/ChaosVehicleMovementComponent.cpp b/Engine/Plugins/Experimental/ChaosVehiclesPlugin/Source/ChaosVehicles/Private/ChaosVehicleMovementComponent.cpp
--- a/Engine/Plugins/Experimental/ChaosVehiclesPlugin/Source/ChaosVehicles/Private/ChaosVehicleMovementComponent.cpp
+++ b/Engine/Plugins/Experimental/ChaosVehiclesPlugin/Source/ChaosVehicles/Private/ChaosVehicleMovementComponent.cpp
@@ -1696,7 +1696,6 @@ TUniquePtr<FChaosVehicleAsyncOutput> FChaosVehicleDefaultAsyncInput::Simulate(UW
 /************************************************************************/
 void UChaosVehicleMovementComponent::Update(float DeltaTime)
 {
-#if WITH_CHAOS
 	if (CurAsyncInput)
 	{
 		if (const FBodyInstance* BodyInstance = GetBodyInstance())
@@ -1730,7 +1729,6 @@ void UChaosVehicleMovementComponent::Update(float DeltaTime)
 			}
 		}
 	}
-#endif
 }
 
 void UChaosVehicleMovementComponent::ResetVehicleState()
diff --git a/Engine/Plugins/Experimental/ChaosVehiclesPlugin/Source/ChaosVehicles/Private/ChaosVehicleWheel.cpp b/Engine/Plugins/Experimental/ChaosVehiclesPlugin/Source/ChaosVehicles/Private/ChaosVehicleWheel.cpp
--- a/Engine/Plugins/Experimental/ChaosVehiclesPlugin/Source/ChaosVehicles/Private/ChaosVehicleWheel.cpp
+++ b/Engine/Plugins/Experimental/ChaosVehiclesPlugin/Source/ChaosVehicles/Private/ChaosVehicleWheel.cpp
@@ -120,19 +120,6 @@ void UChaosVehicleWheel::Init( UChaosWheeledVehicleMovementComponent* InVehicleS
 	VehicleComponent = InVehicleSim;
 	WheelIndex = InWheelIndex;
 
-//#if WITH_PHYSX_VEHICLES
-//	WheelShape = NULL;
-//
-//	FChaosVehicleManager* VehicleManager = FChaosVehicleManager::GetVehicleManagerFromScene(VehicleSim->GetWorld()->GetPhysicsScene());
-//	SCOPED_SCENE_READ_LOCK(VehicleManager->GetScene());
-//
-//	const int32 WheelShapeIdx = VehicleSim->PVehicle->mWheelsSimData.getWheelShapeMapping( WheelIndex );
-//	check(WheelShapeIdx >= 0);
-//
-//	VehicleSim->PVehicle->getRigidDynamicActor()->getShapes( &WheelShape, 1, WheelShapeIdx );
-//	check(WheelShape);
-//#endif // WITH_PHYSX
-
 	Location = GetPhysicsLocation();
 	OldLocation = Location;
 }
diff --git a/Engine/Plugins/Experimental/ChaosVehiclesPlugin/Source/ChaosVehicles/Private/ChaosVehiclesPlugin.cpp b/Engine/Plugins/Experimental/ChaosVehiclesPlugin/Source/ChaosVehicles/Private/ChaosVehiclesPlugin.cpp
--- a/Engine/Plugins/Experimental/ChaosVehiclesPlugin/Source/ChaosVehicles/Private/ChaosVehiclesPlugin.cpp
+++ b/Engine/Plugins/Experimental/ChaosVehiclesPlugin/Source/ChaosVehicles/Private/ChaosVehiclesPlugin.cpp
@@ -16,14 +16,11 @@ IMPLEMENT_MODULE( IChaosVehiclesPlugin, ChaosVehicles )
 
 void IChaosVehiclesPlugin::PhysSceneInit(FPhysScene* PhysScene)
 {
-#if WITH_CHAOS
 	new FChaosVehicleManager(PhysScene);
-#endif
 }
 
 void IChaosVehiclesPlugin::PhysSceneTerm(FPhysScene* PhysScene)
 {
-#if WITH_CHAOS
 	FChaosVehicleManager* VehicleManager = FChaosVehicleManager::GetVehicleManagerFromScene(PhysScene);
 	if (VehicleManager != nullptr)
 	{
@@ -31,12 +28,8 @@ void IChaosVehiclesPlugin::PhysSceneTerm(FPhysScene* PhysScene)
 		delete VehicleManager;
 		VehicleManager = nullptr;
 	}
-#endif // WITH_CHAOS
 }
 
-
-
-
 void IChaosVehiclesPlugin::StartupModule()
 {
 	OnPhysSceneInitHandle = FPhysicsDelegates::OnPhysSceneInit.AddRaw(this, &IChaosVehiclesPlugin::PhysSceneInit);
diff --git a/Engine/Plugins/Experimental/GeometryCollectionPlugin/Source/GeometryCollectionEditor/Private/GeometryCollection/GeometryCollectionConversion.cpp b/Engine/Plugins/Experimental/GeometryCollectionPlugin/Source/GeometryCollectionEditor/Private/GeometryCollection/GeometryCollectionConversion.cpp
--- a/Engine/Plugins/Experimental/GeometryCollectionPlugin/Source/GeometryCollectionEditor/Private/GeometryCollection/GeometryCollectionConversion.cpp
+++ b/Engine/Plugins/Experimental/GeometryCollectionPlugin/Source/GeometryCollectionEditor/Private/GeometryCollection/GeometryCollectionConversion.cpp
@@ -265,12 +265,7 @@ void FGeometryCollectionConversion::AppendStaticMesh(const UStaticMesh* StaticMe
 			CreateGeometryParams.LocalTransform = Chaos::FRigidTransform3::Identity;
 			CreateGeometryParams.WorldTransform = Chaos::FRigidTransform3::Identity;
 			CreateGeometryParams.Geometry = &BodySetup->AggGeom;
-#if PHYSICS_INTERFACE_PHYSX
-			AddParams.TriMeshes = TArrayView<PxTriangleMesh*>(BodySetup->TriMeshes);
-#endif
-#if WITH_CHAOS
 			CreateGeometryParams.ChaosTriMeshes = MakeArrayView(BodySetup->ChaosTriMeshes);
-#endif
 
 			// todo(chaos) : this currently also create the shape array which is unnecessary ,this could be optimized by having a common function to create only the implicits 
 			ChaosInterface::CreateGeometry(CreateGeometryParams, Geoms, Shapes);
diff --git a/Engine/Plugins/Experimental/Water/Source/Runtime/Private/BuoyancyComponent.cpp b/Engine/Plugins/Experimental/Water/Source/Runtime/Private/BuoyancyComponent.cpp
--- a/Engine/Plugins/Experimental/Water/Source/Runtime/Private/BuoyancyComponent.cpp
+++ b/Engine/Plugins/Experimental/Water/Source/Runtime/Private/BuoyancyComponent.cpp
@@ -161,7 +161,6 @@ void UBuoyancyComponent::Update(float DeltaTime)
 		}
 	}
 
-#if WITH_CHAOS
 	if (CurAsyncInput)
 	{
 		if (const FBodyInstance* BodyInstance = SimulatingComponent->GetBodyInstance())
@@ -191,7 +190,6 @@ void UBuoyancyComponent::Update(float DeltaTime)
 			BuoyancyInputState->SmoothedWorldTimeSeconds = GetWorld()->GetTimeSeconds();
 		}
 	}
-#endif
 
 	if (!IsUsingAsyncPath())
 	{
@@ -906,7 +904,6 @@ void UBuoyancyComponent::GameThread_ProcessIntermediateAsyncOutput(const FBuoyan
 
 bool UBuoyancyComponent::IsUsingAsyncPath() const
 {
-#if WITH_CHAOS
 	bool bAsyncSolver = false;
 	if (UWorld* World = GetWorld())
 	{
@@ -919,8 +916,6 @@ bool UBuoyancyComponent::IsUsingAsyncPath() const
 		}
 	}
 	return bAsyncSolver && bUseAsyncPath && (CVarWaterBuoyancyUseAsyncPath.GetValueOnAnyThread() > 0);
-#endif
-	return false;
 }
 
 TUniquePtr<FBuoyancyComponentAsyncAux> UBuoyancyComponent::CreateAsyncAux() const
diff --git a/Engine/Plugins/Experimental/Water/Source/Runtime/Private/BuoyancyManager.cpp b/Engine/Plugins/Experimental/Water/Source/Runtime/Private/BuoyancyManager.cpp
--- a/Engine/Plugins/Experimental/Water/Source/Runtime/Private/BuoyancyManager.cpp
+++ b/Engine/Plugins/Experimental/Water/Source/Runtime/Private/BuoyancyManager.cpp
@@ -28,7 +28,6 @@ ABuoyancyManager* ABuoyancyManager::Get(const UObject* WorldContextObject)
 
 void ABuoyancyManager::OnCreatePhysics(UActorComponent* Component)
 {
-#if WITH_CHAOS
 	if (AActor* OwningActor = Component->GetOwner())
 	{
 		if (UBuoyancyComponent* BuoyancyComp = OwningActor->FindComponentByClass<UBuoyancyComponent>())
@@ -39,12 +38,10 @@ void ABuoyancyManager::OnCreatePhysics(UActorComponent* Component)
 			}
 		}
 	}
-#endif
 }
 
 void ABuoyancyManager::OnDestroyPhysics(UActorComponent* Component)
 {
-#if WITH_CHAOS
 	if (AActor* OwningActor = Component->GetOwner())
 	{
 		if (UBuoyancyComponent* BuoyancyComp = OwningActor->FindComponentByClass<UBuoyancyComponent>())
@@ -70,7 +67,6 @@ void ABuoyancyManager::OnDestroyPhysics(UActorComponent* Component)
 			}
 		}
 	}
-#endif
 }
 
 void ABuoyancyManager::ClearAsyncInputs(UBuoyancyComponent* Component)
@@ -221,7 +217,6 @@ void ABuoyancyManager::Update(FPhysScene* PhysScene, float DeltaTime)
 
 void ABuoyancyManager::InitializeAsyncAux(UBuoyancyComponent* Component)
 {
-#if WITH_CHAOS
 	UPrimitiveComponent* SimulatingComponent = Component->GetSimulatingComponent();
 	if (!SimulatingComponent)
 	{
@@ -243,7 +238,6 @@ void ABuoyancyManager::InitializeAsyncAux(UBuoyancyComponent* Component)
 			}
 		}
 	}
-#endif
 }
 
 void ABuoyancyManager::BeginPlay()
@@ -256,10 +250,7 @@ void ABuoyancyManager::BeginPlay()
 		if (FPhysScene* PhysScene = World->GetPhysicsScene())
 		{
 			OnPhysScenePreTickHandle = PhysScene->OnPhysScenePreTick.AddUObject(this, &ABuoyancyManager::Update);
-#if WITH_CHAOS
-			OnPhysScenePreTickHandle = PhysScene->OnPhysScenePreTick.AddUObject(this, &ABuoyancyManager::Update);
 			AsyncCallback = PhysScene->GetSolver()->CreateAndRegisterSimCallbackObject_External<FBuoyancyManagerAsyncCallback>();
-#endif
 		}
 	}
 
@@ -281,14 +272,13 @@ void ABuoyancyManager::EndPlay(const EEndPlayReason::Type EndPlayReason)
 		if (FPhysScene* PhysScene = World->GetPhysicsScene())
 		{
 			PhysScene->OnPhysScenePreTick.Remove(OnPhysScenePreTickHandle);
-#if WITH_CHAOS
+
 			if (AsyncCallback)
 			{
 				PhysScene->GetSolver()->UnregisterAndFreeSimCallbackObject_External(AsyncCallback);
 				PhysScene->OnPhysScenePreTick.Remove(OnPhysScenePreTickHandle);
 				AsyncCallback = nullptr;
 			}
-#endif
 		}
 	}
 }
diff --git a/Engine/Plugins/Runtime/ApexDestruction/Source/ApexDestruction/Public/DestructibleComponent.h b/Engine/Plugins/Runtime/ApexDestruction/Source/ApexDestruction/Public/DestructibleComponent.h
--- a/Engine/Plugins/Runtime/ApexDestruction/Source/ApexDestruction/Public/DestructibleComponent.h
+++ b/Engine/Plugins/Runtime/ApexDestruction/Source/ApexDestruction/Public/DestructibleComponent.h
@@ -17,14 +17,6 @@ struct FBodyInstance;
 struct FCollisionShape;
 struct FNavigableGeometryExport;
 
-#if WITH_PHYSX
-namespace physx
-{
-	class PxRigidDynamic;
-	class PxRigidActor;
-}
-#endif // WITH_PHYSX 
-
 /** Delegate for notification when fracture occurs */
 DECLARE_DYNAMIC_MULTICAST_DELEGATE_TwoParams(FComponentFractureSignature, const FVector &, HitPoint, const FVector &, HitDirection);
 
@@ -193,19 +185,6 @@ PRAGMA_ENABLE_DEPRECATION_WARNINGS
 	/** Collision response used for chunks */
 	FCollisionResponse LargeChunkCollisionResponse;
 	FCollisionResponse SmallChunkCollisionResponse;
-#if WITH_PHYSX
-	/** User data wrapper for this component passed to physx */
-	FPhysxUserData PhysxUserData;
-
-	void SetCollisionResponseForShape(physx::PxShape* Shape, int32 ChunkIdx);
-	void SetCollisionResponseForActor(physx::PxRigidDynamic* Actor, int32 ChunkIdx, const FCollisionResponseContainer* ResponseOverride = NULL);
-
-
-public:
-	/** User data wrapper for the chunks passed to physx */
-	TArray<FPhysxUserData> PhysxChunkUserData;
-	bool IsChunkLarge(physx::PxRigidActor* ChunkActor) const;
-#endif
 
 private:
 	/** Cached values for computing contact offsets */
diff --git a/Engine/Plugins/Runtime/MeshModelingToolset/Source/ModelingComponents/Private/Physics/PhysicsDataCollection.cpp b/Engine/Plugins/Runtime/MeshModelingToolset/Source/ModelingComponents/Private/Physics/PhysicsDataCollection.cpp
--- a/Engine/Plugins/Runtime/MeshModelingToolset/Source/ModelingComponents/Private/Physics/PhysicsDataCollection.cpp
+++ b/Engine/Plugins/Runtime/MeshModelingToolset/Source/ModelingComponents/Private/Physics/PhysicsDataCollection.cpp
@@ -112,18 +112,6 @@ void FPhysicsDataCollection::CopyGeometryToAggregate()
 		FKConvexElem Element;
 		UE::Geometry::GetFKElement(ConvexGeom.Mesh, Element);
 
-#if !WITH_CHAOS
-		// Chaos will compute the IndexData itself on the call to FKConvexElem::UpdateElemBox() in ::GetFKElement() above.
-		// PhysX will not, so initialize that data with the mesh triangles.
-		// (This code should go into ::GetFKElement but cannot because it needs to be added in a hotfix)
-		for (FIndex3i Triangle : ConvexGeom.Mesh.TrianglesItr())
-		{
-			Element.IndexData.Add(Triangle.A);
-			Element.IndexData.Add(Triangle.B);
-			Element.IndexData.Add(Triangle.C);
-		}
-#endif
-
 		AggGeom.ConvexElems.Add(Element);
 	}
 
diff --git a/Engine/Plugins/Runtime/NetworkPrediction/Source/NetworkPrediction/Private/NetworkPredictionPhysics.cpp b/Engine/Plugins/Runtime/NetworkPrediction/Source/NetworkPrediction/Private/NetworkPredictionPhysics.cpp
--- a/Engine/Plugins/Runtime/NetworkPrediction/Source/NetworkPrediction/Private/NetworkPredictionPhysics.cpp
+++ b/Engine/Plugins/Runtime/NetworkPrediction/Source/NetworkPrediction/Private/NetworkPredictionPhysics.cpp
@@ -6,8 +6,6 @@
 #include "Components/PrimitiveComponent.h"
 #include "PhysicsProxy/SingleParticlePhysicsProxy.h"
 
-#if WITH_CHAOS
-
 // -------------------------------------------------------------------------------------------------------------------------
 //	Interpolation related functions. These require calls to the UPrimitiveComponent and cannot be implemented via FBodyInstance
 //
@@ -174,4 +172,3 @@ void FNetworkPredictionPhysicsState::ToString(FBodyInstance* BodyInstance, FAnsi
 		Builder.Append("Null BodyInstance\n");
 	}
 }
-#endif // WITH_CHAOS
diff --git a/Engine/Plugins/Runtime/NetworkPrediction/Source/NetworkPrediction/Private/NetworkPredictionWorldManager.cpp b/Engine/Plugins/Runtime/NetworkPrediction/Source/NetworkPrediction/Private/NetworkPredictionWorldManager.cpp
--- a/Engine/Plugins/Runtime/NetworkPrediction/Source/NetworkPrediction/Private/NetworkPredictionWorldManager.cpp
+++ b/Engine/Plugins/Runtime/NetworkPrediction/Source/NetworkPrediction/Private/NetworkPredictionWorldManager.cpp
@@ -580,7 +580,7 @@ void UNetworkPredictionWorldManager::SetUsingPhysics()
 	{
 		npCheckSlow(Physics.Solver == nullptr);
 		Physics.bUsingPhysics = true;
-#if WITH_CHAOS
+
 		Physics.Module = FChaosSolversModule::GetModule();
 		Physics.Solver = GetWorld()->GetPhysicsScene()->GetSolver();
 
@@ -602,7 +602,6 @@ void UNetworkPredictionWorldManager::SetUsingPhysics()
 			GEngine->bUseFixedFrameRate = true;
 			GEngine->FixedFrameRate = Settings.FixedTickFrameRate;
 		}
-#endif // WITH_CHAOS
 	}
 }
 
diff --git a/Engine/Plugins/Runtime/NetworkPrediction/Source/NetworkPrediction/Public/NetworkPredictionPhysics.h b/Engine/Plugins/Runtime/NetworkPrediction/Source/NetworkPrediction/Public/NetworkPredictionPhysics.h
--- a/Engine/Plugins/Runtime/NetworkPrediction/Source/NetworkPrediction/Public/NetworkPredictionPhysics.h
+++ b/Engine/Plugins/Runtime/NetworkPrediction/Source/NetworkPrediction/Public/NetworkPredictionPhysics.h
@@ -19,8 +19,6 @@
 #include "PhysicsEngine/BodyInstance.h"
 #include "Containers/StringConv.h"
 
-#if WITH_CHAOS
-
 namespace NetworkPredictionPhysicsCvars
 {
 	NETSIM_DEVCVAR_SHIPCONST_INT(FullPrecision, 1, "np.Physics.FullPrecision", "Replicatre physics state with full precision. Not to be toggled during gameplay.");
@@ -161,8 +159,6 @@ private:
 	}
 };
 
-#endif // WITH_CHAOS
-
 // ------------------------------------------------------------------------------------------------------
 // Generic model def for physics actors with no backing simulation
 // ------------------------------------------------------------------------------------------------------
@@ -170,10 +166,7 @@ struct NETWORKPREDICTION_API FGenericPhysicsModelDef : FNetworkPredictionModelDe
 {
 	NP_MODEL_BODY();
 
-#if WITH_CHAOS
-
 	using PhysicsState = FNetworkPredictionPhysicsState;
-#endif // WITH_CHAOS
 	using Driver = UPrimitiveComponent;
 
 	static const TCHAR* GetName() { return TEXT("Generic Physics Actor"); }
diff --git a/Engine/Plugins/Runtime/NetworkPrediction/Source/NetworkPrediction/Public/NetworkPredictionPhysicsComponent.h b/Engine/Plugins/Runtime/NetworkPrediction/Source/NetworkPrediction/Public/NetworkPredictionPhysicsComponent.h
--- a/Engine/Plugins/Runtime/NetworkPrediction/Source/NetworkPrediction/Public/NetworkPredictionPhysicsComponent.h
+++ b/Engine/Plugins/Runtime/NetworkPrediction/Source/NetworkPrediction/Public/NetworkPredictionPhysicsComponent.h
@@ -58,9 +58,7 @@ protected:
 	UPROPERTY(Replicated, transient)
 	FReplicationProxy ReplicationProxy;
 
-#if WITH_CHAOS
 	FPhysicsActorHandle PhysicsActorHandle = nullptr;
-#endif // WITH_CHAOS
 
 	FReplicationProxySet GetReplicationProxies()
 	{
diff --git a/Engine/Plugins/Runtime/NetworkPredictionExtras/Source/NetworkPredictionExtras/Private/MockPhysicsComponent.cpp b/Engine/Plugins/Runtime/NetworkPredictionExtras/Source/NetworkPredictionExtras/Private/MockPhysicsComponent.cpp
--- a/Engine/Plugins/Runtime/NetworkPredictionExtras/Source/NetworkPredictionExtras/Private/MockPhysicsComponent.cpp
+++ b/Engine/Plugins/Runtime/NetworkPredictionExtras/Source/NetworkPredictionExtras/Private/MockPhysicsComponent.cpp
@@ -28,12 +28,10 @@ public:
 
 	NP_MODEL_BODY();
 
-#if WITH_CHAOS
 	using Simulation = FMockPhysicsSimulation;
 	using StateTypes = MockPhysicsStateTypes;
 	using Driver = UMockPhysicsComponent;
 	using PhysicsState = FNetworkPredictionPhysicsState;
-#endif // WITH_CHAOS
 
 	static const TCHAR* GetName() { return TEXT("MockPhysics"); }
 	static constexpr int32 GetSortPriority() { return (int32)ENetworkPredictionSortPriority::Physics+1; }
@@ -45,10 +43,8 @@ NP_MODEL_REGISTER(FMockPhysicsModelDef);
 //	UMockPhysicsComponent
 // ----------------------------------------------------------------------------------------------------------
 
-#if WITH_CHAOS
 // UMockPhysicsComponent implements the FMockPhysicsCueSet
 NETSIMCUESET_REGISTER(UMockPhysicsComponent, FMockPhysicsCueSet);
-#endif // WITH_CHAOS
 
 UMockPhysicsComponent::UMockPhysicsComponent()
 {
@@ -57,7 +53,6 @@ UMockPhysicsComponent::UMockPhysicsComponent()
 
 void UMockPhysicsComponent::InitializeNetworkPredictionProxy()
 {
-#if WITH_CHAOS
 	// We need valid UpdatedPrimitive and PhysicsHandle to register
 	// This code does not handle any "not ready yet" cases
 	if (npEnsure(UpdatedPrimitive))
@@ -69,7 +64,6 @@ void UMockPhysicsComponent::InitializeNetworkPredictionProxy()
 		ActiveSimulation->PrimitiveComponent = UpdatedPrimitive;
 		NetworkPredictionProxy.Init<FMockPhysicsModelDef>(GetWorld(), GetReplicationProxies(), ActiveSimulation, this);
 	}
-#endif
 }
 
 void UMockPhysicsComponent::ProduceInput(const int32 DeltaTimeMS, FMockPhysicsInputCmd* Cmd)
diff --git a/Engine/Plugins/Runtime/NetworkPredictionExtras/Source/NetworkPredictionExtras/Private/MockPhysicsGrenade.cpp b/Engine/Plugins/Runtime/NetworkPredictionExtras/Source/NetworkPredictionExtras/Private/MockPhysicsGrenade.cpp
--- a/Engine/Plugins/Runtime/NetworkPredictionExtras/Source/NetworkPredictionExtras/Private/MockPhysicsGrenade.cpp
+++ b/Engine/Plugins/Runtime/NetworkPredictionExtras/Source/NetworkPredictionExtras/Private/MockPhysicsGrenade.cpp
@@ -48,10 +48,7 @@ void UMockPhysicsGrenadeComponent::InitializeNetworkPredictionProxy()
 {
 	npCheckSlow(this->UpdatedPrimitive);
 	
-	// FIXME: this can't compile without WITH_CHAOS
-#if WITH_CHAOS
 	NetworkPredictionProxy.Init<FMockPhysicsGrenadeModelDef>(GetWorld(), GetReplicationProxies(), this, this);
-#endif
 }
 
 void UMockPhysicsGrenadeComponent::InitializeSimulationState(void* Sync, FMockPhysicsGrenadeAuxState* Aux)
diff --git a/Engine/Plugins/Runtime/NetworkPredictionExtras/Source/NetworkPredictionExtras/Private/MockPhysicsSimulation.cpp b/Engine/Plugins/Runtime/NetworkPredictionExtras/Source/NetworkPredictionExtras/Private/MockPhysicsSimulation.cpp
--- a/Engine/Plugins/Runtime/NetworkPredictionExtras/Source/NetworkPredictionExtras/Private/MockPhysicsSimulation.cpp
+++ b/Engine/Plugins/Runtime/NetworkPredictionExtras/Source/NetworkPredictionExtras/Private/MockPhysicsSimulation.cpp
@@ -19,8 +19,6 @@
 #include "NetworkPredictionCues.h"
 #include "PhysicsProxy/SingleParticlePhysicsProxy.h"
 
-#if WITH_CHAOS
-
 NETSIMCUE_REGISTER(FMockPhysicsJumpCue, TEXT("MockPhysicsJumpCue"));
 NETSIMCUE_REGISTER(FMockPhysicsChargeCue, TEXT("FMockPhysicsChargeCue"));
 
@@ -140,5 +138,3 @@ void FMockPhysicsSimulation::SimulationTick(const FNetSimTimeStep& TimeStep, con
 		}
 	}
 }
-
-#endif // WITH_CHAOS
\ No newline at end of file
diff --git a/Engine/Plugins/Runtime/NetworkPredictionExtras/Source/NetworkPredictionExtras/Public/MockPhysicsSimulation.h b/Engine/Plugins/Runtime/NetworkPredictionExtras/Source/NetworkPredictionExtras/Public/MockPhysicsSimulation.h
--- a/Engine/Plugins/Runtime/NetworkPredictionExtras/Source/NetworkPredictionExtras/Public/MockPhysicsSimulation.h
+++ b/Engine/Plugins/Runtime/NetworkPredictionExtras/Source/NetworkPredictionExtras/Public/MockPhysicsSimulation.h
@@ -162,8 +162,6 @@ struct FMockPhysicsChargeCue
 	}
 };
 
-#if WITH_CHAOS
-
 struct FMockPhysicsCueSet
 {
 	template<typename TDispatchTable>
@@ -173,5 +171,3 @@ struct FMockPhysicsCueSet
 		DispatchTable.template RegisterType<FMockPhysicsChargeCue>();
 	}
 };
-
-#endif // WITH_CHAOS
\ No newline at end of file
diff --git a/Engine/Source/Developer/AssetTools/Private/AssetTypeActions/AssetTypeActions_SkeletalMesh.cpp b/Engine/Source/Developer/AssetTools/Private/AssetTypeActions/AssetTypeActions_SkeletalMesh.cpp
--- a/Engine/Source/Developer/AssetTools/Private/AssetTypeActions/AssetTypeActions_SkeletalMesh.cpp
+++ b/Engine/Source/Developer/AssetTools/Private/AssetTypeActions/AssetTypeActions_SkeletalMesh.cpp
@@ -474,15 +474,6 @@ void FAssetTypeActions_SkeletalMesh::GetActions(const TArray<UObject*>& InObject
 		false,
 		FSlateIcon(FAppStyle::GetAppStyleSetName(), "Icons.LOD")
 		);
-	
-#if WITH_APEX_CLOTHING
-	Section.AddMenuEntry(
-		"ImportClothing_Entry",
-		LOCTEXT("ImportClothing_Entry", "Import Clothing Asset..."),
-		LOCTEXT("ImportClothing_ToolTip", "Import a clothing asset from a supported file on disk into this skeletal mesh."),
-		FSlateIcon(),
-		FUIAction(FExecuteAction::CreateSP(this, &FAssetTypeActions_SkeletalMesh::ExecuteImportClothing, Meshes)));
-#endif  // #if WITH_APEX_CLOTHING
 
 	// skeleton menu
 	Section.AddSubMenu(
@@ -854,21 +845,6 @@ void FAssetTypeActions_SkeletalMesh::ExecuteImportMeshLOD(UObject* Mesh, int32 L
 	}
 }
 
-#if WITH_APEX_CLOTHING
-void FAssetTypeActions_SkeletalMesh::ExecuteImportClothing(TArray<TWeakObjectPtr<USkeletalMesh>> Objects)
-{
-	if(Objects.Num() > 0)
-	{
-		USkeletalMesh* TargetMesh = Objects[0].Get();
-
-		if(TargetMesh)
-		{
-			ApexClothingUtils::PromptAndImportClothing(TargetMesh);
-		}
-	}
-}
-#endif  // #if WITH_APEX_CLOTHING
-
 void FAssetTypeActions_SkeletalMesh::FillSkeletonMenu(FMenuBuilder& MenuBuilder, TArray<TWeakObjectPtr<USkeletalMesh>> Meshes) const
 {
 	MenuBuilder.BeginSection("SkeletonMenu", LOCTEXT("SkeletonMenuHeading", "Skeleton"));
diff --git a/Engine/Source/Developer/AssetTools/Private/AssetTypeActions/AssetTypeActions_SkeletalMesh.h b/Engine/Source/Developer/AssetTools/Private/AssetTypeActions/AssetTypeActions_SkeletalMesh.h
--- a/Engine/Source/Developer/AssetTools/Private/AssetTypeActions/AssetTypeActions_SkeletalMesh.h
+++ b/Engine/Source/Developer/AssetTools/Private/AssetTypeActions/AssetTypeActions_SkeletalMesh.h
@@ -55,11 +55,6 @@ private:
 	/** Handler for skeletal mesh import */
 	static void ExecuteImportMeshLOD(class UObject* Mesh, int32 LOD);
 
-#if WITH_APEX_CLOTHING
-	/** Handler to import clothing to a mesh */
-	void ExecuteImportClothing(TArray<TWeakObjectPtr<USkeletalMesh>> Objects);
-#endif
-
 	// Helper functions
 private:
 	/** Assigns a skeleton to the mesh */
diff --git a/Engine/Source/Developer/MeshMergeUtilities/Private/MeshMergeHelpers.cpp b/Engine/Source/Developer/MeshMergeUtilities/Private/MeshMergeHelpers.cpp
--- a/Engine/Source/Developer/MeshMergeUtilities/Private/MeshMergeHelpers.cpp
+++ b/Engine/Source/Developer/MeshMergeUtilities/Private/MeshMergeHelpers.cpp
@@ -949,12 +949,7 @@ void FMeshMergeHelpers::ExtractPhysicsGeometry(UBodySetup* InBodySetup, const FT
 	OutAggGeom.RenderInfo = nullptr;
 	for (FKConvexElem& Elem : OutAggGeom.ConvexElems)
 	{
-#if PHYSICS_INTERFACE_PHYSX
-		Elem.SetConvexMesh(nullptr);
-		Elem.SetMirroredConvexMesh(nullptr);
-#elif WITH_CHAOS
 		Elem.ResetChaosConvexMesh();
-#endif
 	}
 
 	// Transform geometry to world space
diff --git a/Engine/Source/Developer/SkeletalMeshUtilitiesCommon/Private/LODUtilities.cpp b/Engine/Source/Developer/SkeletalMeshUtilitiesCommon/Private/LODUtilities.cpp
--- a/Engine/Source/Developer/SkeletalMeshUtilitiesCommon/Private/LODUtilities.cpp
+++ b/Engine/Source/Developer/SkeletalMeshUtilitiesCommon/Private/LODUtilities.cpp
@@ -29,10 +29,6 @@
 
 #include "ObjectTools.h"
 
-#if WITH_APEX_CLOTHING
-	#include "ApexClothingUtils.h"
-#endif // #if WITH_APEX_CLOTHING
-
 #include "ComponentReregisterContext.h"
 #include "IMeshReductionManagerModule.h"
 #include "Animation/SkinWeightProfile.h"
diff --git a/Engine/Source/Developer/TargetPlatform/Private/TargetPlatformManagerModule.cpp b/Engine/Source/Developer/TargetPlatform/Private/TargetPlatformManagerModule.cpp
--- a/Engine/Source/Developer/TargetPlatform/Private/TargetPlatformManagerModule.cpp
+++ b/Engine/Source/Developer/TargetPlatform/Private/TargetPlatformManagerModule.cpp
@@ -25,11 +25,6 @@
 #include "DesktopPlatformModule.h"
 #include "Interfaces/ITurnkeySupportModule.h"
 
-#if PHYSICS_INTERFACE_PHYSX
-#include "IPhysXCooking.h"
-#include "IPhysXCookingModule.h"
-#endif // PHYSICS_INTERFACE_PHYSX
-
 DEFINE_LOG_CATEGORY_STATIC(LogTargetPlatformManager, Log, All);
 
 // AutoSDKs needs the extra DDPI info
@@ -682,59 +677,11 @@ public:
 		static bool bInitialized = false;
 		static TArray<const IPhysXCooking*> Results;
 
-#if PHYSICS_INTERFACE_PHYSX
-		if (!bInitialized || bForceCacheUpdate)
-		{
-			bInitialized = true;
-			Results.Empty(Results.Num());
-			
-			TArray<FName> Modules;
-			FModuleManager::Get().FindModules(TEXT("PhysXCooking*"), Modules);
-			
-			if (!Modules.Num())
-			{
-				UE_LOG(LogTargetPlatformManager, Error, TEXT("No target PhysX formats found!"));
-			}
-
-			for (int32 Index = 0; Index < Modules.Num(); Index++)
-			{
-				IPhysXCookingModule* Module = FModuleManager::LoadModulePtr<IPhysXCookingModule>(Modules[Index]);
-				if (Module)
-				{
-					IPhysXCooking* Format = Module->GetPhysXCooking();
-					if (Format != nullptr)
-					{
-						Results.Add(Format);
-					}
-				}
-			}
-		}
-#endif // PHYSICS_INTERFACE_PHYSX
-
 		return Results;
 	}
 
 	virtual const IPhysXCooking* FindPhysXCooking(FName Name) override
 	{
-#if PHYSICS_INTERFACE_PHYSX 
-		const TArray<const IPhysXCooking*>& PhysXCooking = GetPhysXCooking();
-
-		for (int32 Index = 0; Index < PhysXCooking.Num(); Index++)
-		{
-			TArray<FName> Formats;
-
-			PhysXCooking[Index]->GetSupportedFormats(Formats);
-		
-			for (int32 FormatIndex = 0; FormatIndex < Formats.Num(); FormatIndex++)
-			{
-				if (Formats[FormatIndex] == Name)
-				{
-					return PhysXCooking[Index];
-				}
-			}
-		}
-#endif // PHYSICS_INTERFACE_PHYSX
-
 		return nullptr;
 	}
 
diff --git a/Engine/Source/Developer/TargetPlatform/TargetPlatform.Build.cs b/Engine/Source/Developer/TargetPlatform/TargetPlatform.Build.cs
--- a/Engine/Source/Developer/TargetPlatform/TargetPlatform.Build.cs
+++ b/Engine/Source/Developer/TargetPlatform/TargetPlatform.Build.cs
@@ -89,10 +89,5 @@ public class TargetPlatform : ModuleRules
 				}
 			}
 		}
-
-		if (Target.bBuildDeveloperTools == true && Target.bBuildRequiresCookedData && Target.bCompileAgainstEngine && Target.bCompilePhysX)
-		{
-			DynamicallyLoadedModuleNames.Add("PhysXCooking");
-		}
 	}
 }
diff --git a/Engine/Source/Editor/AnimGraph/Private/AnimGraphNode_RigidBody.cpp b/Engine/Source/Editor/AnimGraph/Private/AnimGraphNode_RigidBody.cpp
--- a/Engine/Source/Editor/AnimGraph/Private/AnimGraphNode_RigidBody.cpp
+++ b/Engine/Source/Editor/AnimGraph/Private/AnimGraphNode_RigidBody.cpp
@@ -40,13 +40,6 @@ FText UAnimGraphNode_RigidBody::GetNodeTitle(ENodeTitleType::Type TitleType) con
 
 void UAnimGraphNode_RigidBody::ValidateAnimNodeDuringCompilation(USkeleton* ForSkeleton, FCompilerResultsLog& MessageLog)
 {
-#if !WITH_CHAOS
-	if(Node.bEnableWorldGeometry && Node.SimulationSpace != ESimulationSpace::WorldSpace)
-	{
-		MessageLog.Error(*LOCTEXT("AnimGraphNode_CompileError", "@@ - uses world collision without world space simulation. This is not supported").ToString());
-	}
-#endif
-	
 	Super::ValidateAnimNodeDuringCompilation(ForSkeleton, MessageLog);
 }
 
diff --git a/Engine/Source/Editor/Cascade/Private/Cascade.cpp b/Engine/Source/Editor/Cascade/Private/Cascade.cpp
--- a/Engine/Source/Editor/Cascade/Private/Cascade.cpp
+++ b/Engine/Source/Editor/Cascade/Private/Cascade.cpp
@@ -1621,12 +1621,7 @@ void FCascade::Tick(float DeltaTime)
 		World->SetupPhysicsTickFunctions(DeltaTime);
 		PhysScene->StartFrame();
 		PhysScene->WaitPhysScenes();
-
-#if WITH_CHAOS
 		PhysScene->EndFrame();
-#else
-		PhysScene->EndFrame(NULL);
-#endif
 	}
 
 	// If a vector field module is selected, update the preview visualization.
diff --git a/Engine/Source/Editor/ClothPainter/Private/ClothingAssetListCommands.cpp b/Engine/Source/Editor/ClothPainter/Private/ClothingAssetListCommands.cpp
--- a/Engine/Source/Editor/ClothPainter/Private/ClothingAssetListCommands.cpp
+++ b/Engine/Source/Editor/ClothPainter/Private/ClothingAssetListCommands.cpp
@@ -9,9 +9,6 @@
 void FClothingAssetListCommands::RegisterCommands()
 {
 	UI_COMMAND(DeleteAsset, "Delete Asset", "Deletes a clothing asset from the mesh.", EUserInterfaceActionType::Button, FInputChord());
-#if WITH_APEX_CLOTHING
-	UI_COMMAND(ReimportAsset, "Reimport Asset", "If an asset was originally imported from an external file, reimport from that file.", EUserInterfaceActionType::Button, FInputChord());
-#endif
 	UI_COMMAND(RebuildAssetParams, "Rebuild Asset Parameter Masks", "Takes the parameter masks in LOD0 and creates masks in all lower LODs to match, casting those parameter masks to the new mesh.", EUserInterfaceActionType::Button, FInputChord());
 
 	// Add one command per exporter provider
diff --git a/Engine/Source/Editor/ClothPainter/Private/ClothingAssetListCommands.h b/Engine/Source/Editor/ClothPainter/Private/ClothingAssetListCommands.h
--- a/Engine/Source/Editor/ClothPainter/Private/ClothingAssetListCommands.h
+++ b/Engine/Source/Editor/ClothPainter/Private/ClothingAssetListCommands.h
@@ -21,10 +21,6 @@ public:
 
 	TSharedPtr<FUICommandInfo> DeleteAsset;
 
-#if WITH_APEX_CLOTHING
-	TSharedPtr<FUICommandInfo> ReimportAsset;
-#endif
-
 	TSharedPtr<FUICommandInfo> RebuildAssetParams;
 	TMap<FName, TSharedPtr<FUICommandInfo>> ExportAssets;
 };
diff --git a/Engine/Source/Editor/ClothPainter/Private/SClothAssetSelector.cpp b/Engine/Source/Editor/ClothPainter/Private/SClothAssetSelector.cpp
--- a/Engine/Source/Editor/ClothPainter/Private/SClothAssetSelector.cpp
+++ b/Engine/Source/Editor/ClothPainter/Private/SClothAssetSelector.cpp
@@ -148,14 +148,6 @@ public:
 			FExecuteAction::CreateSP(this, &SAssetListRow::DeleteAsset)
 		);
 
-#if WITH_APEX_CLOTHING
-		UICommandList->MapAction(
-			Commands.ReimportAsset,
-			FExecuteAction::CreateSP(this, &SAssetListRow::ReimportAsset),
-			FCanExecuteAction::CreateSP(this, &SAssetListRow::CanReimportAsset)
-		);
-#endif
-
 		UICommandList->MapAction(
 			Commands.RebuildAssetParams,
 			FExecuteAction::CreateSP(this, &SAssetListRow::RebuildLODParameters),
@@ -184,9 +176,6 @@ public:
 			Builder.BeginSection(NAME_None, LOCTEXT("AssetActions_SectionName", "Actions"));
 			{
 				Builder.AddMenuEntry(FGenericCommands::Get().Delete);
-#if WITH_APEX_CLOTHING
-				Builder.AddMenuEntry(Commands.ReimportAsset);
-#endif
 				Builder.AddMenuEntry(Commands.RebuildAssetParams);
 
 				for (const TPair<FName, TSharedPtr<FUICommandInfo>>& CommandId : Commands.ExportAssets)
@@ -261,62 +250,6 @@ private:
 		}
 	}
 
-#if WITH_APEX_CLOTHING
-	void ReimportAsset()
-	{
-		if(UClothingAssetCommon* Asset = Item->ClothingAsset.Get())
-		{
-			if(USkeletalMesh* SkelMesh = Cast<USkeletalMesh>(Asset->GetOuter()))
-			{
-				FString ReimportPath = Asset->ImportedFilePath;
-
-				if(ReimportPath.IsEmpty())
-				{
-					const FText MessageText = LOCTEXT("Warning_NoReimportPath", "There is no reimport path available for this asset, it was likely created in the Editor. Would you like to select a file and overwrite this asset?");
-					EAppReturnType::Type MessageReturn = FMessageDialog::Open(EAppMsgType::YesNo, MessageText);
-
-					if(MessageReturn == EAppReturnType::Yes)
-					{
-						ReimportPath = ApexClothingUtils::PromptForClothingFile();
-					}
-				}
-
-				if(ReimportPath.IsEmpty())
-				{
-					return;
-				}
-
-				// Retry if the file isn't there
-				if(!FPaths::FileExists(ReimportPath))
-				{
-					const FText MessageText = LOCTEXT("Warning_NoFileFound", "Could not find an asset to reimport, select a new file on disk?");
-					EAppReturnType::Type MessageReturn = FMessageDialog::Open(EAppMsgType::YesNo, MessageText);
-
-					if(MessageReturn == EAppReturnType::Yes)
-					{
-						ReimportPath = ApexClothingUtils::PromptForClothingFile();
-					}
-				}
-
-				FClothingSystemEditorInterfaceModule& ClothingEditorInterface = FModuleManager::Get().LoadModuleChecked<FClothingSystemEditorInterfaceModule>("ClothingSystemEditorInterface");
-				UClothingAssetFactoryBase* Factory = ClothingEditorInterface.GetClothingAssetFactory();
-
-				if(Factory && Factory->CanImport(ReimportPath))
-				{
-					Factory->Reimport(ReimportPath, SkelMesh, Asset);
-
-					OnInvalidateList.ExecuteIfBound();
-				}
-			}
-		}
-	}
-
-	bool CanReimportAsset() const
-	{
-		return Item.IsValid() && !Item->ClothingAsset->ImportedFilePath.IsEmpty();
-	}
-#endif  // #if WITH_APEX_CLOTHING
-
 	// Using LOD0 of an asset, rebuild the other LOD masks by mapping the LOD0 parameters onto their meshes
 	void RebuildLODParameters()
 	{
@@ -789,42 +722,6 @@ void SClothAssetSelector::Construct(const FArguments& InArgs, USkeletalMesh* InM
 					.TextStyle(FAppStyle::Get(), "DetailsView.CategoryTextStyle")
 					.Font(FAppStyle::Get().GetFontStyle("PropertyWindow.BoldFont"))
 				]
-#if WITH_APEX_CLOTHING
-				+SHorizontalBox::Slot()
-				.AutoWidth()
-				.VAlign(VAlign_Center)
-				.HAlign(HAlign_Right)
-				.Padding(0.0f, 0.0f, 4.0f, 0.0f)
-				[
-					SNew(SButton)
-					.OnClicked(this, &SClothAssetSelector::OnImportApexFileClicked)
-					.HAlign(HAlign_Center)
-					.VAlign(VAlign_Center)
-					[
-						SNew(SHorizontalBox)
-
-						+ SHorizontalBox::Slot()
-						.VAlign(VAlign_Center)
-						.AutoWidth()
-						.Padding(FMargin(0, 1))
-						[
-							SNew(SImage)
-							.Image(FAppStyle::GetBrush("Plus"))
-						]
-
-						+ SHorizontalBox::Slot()
-						.VAlign(VAlign_Center)
-						.AutoWidth()
-						.Padding(FMargin(2, 0, 0, 0))
-						[
-							SNew(STextBlock)
-							.Font(IDetailLayoutBuilder::GetDetailFontBold())
-							.Text(LOCTEXT("NewAssetButtonText", "Import APEX file"))
-							.Visibility(this, &SClothAssetSelector::GetAssetHeaderButtonTextVisibility)
-						]
-					]
-				]
-#endif  // #if WITH_APEX_CLOTHING
 				+ SHorizontalBox::Slot()
 				.AutoWidth()
 				.VAlign(VAlign_Center)
@@ -1197,21 +1094,6 @@ void SClothAssetSelector::PostUndo(bool bSuccess)
 	OnRefresh();
 }
 
-#if WITH_APEX_CLOTHING
-FReply SClothAssetSelector::OnImportApexFileClicked()
-{
-	if(Mesh)
-	{
-		ApexClothingUtils::PromptAndImportClothing(Mesh);
-		OnRefresh();
-
-		return FReply::Handled();
-	}
-
-	return FReply::Unhandled();
-}
-#endif  // #if WITH_APEX_CLOTHING
-
 void SClothAssetSelector::OnCopyClothingAssetSelected(const FAssetData& AssetData)
 {
 	USkeletalMesh* SourceSkelMesh = Cast<USkeletalMesh>(AssetData.GetAsset());
diff --git a/Engine/Source/Editor/ClothPainter/Private/SClothAssetSelector.h b/Engine/Source/Editor/ClothPainter/Private/SClothAssetSelector.h
--- a/Engine/Source/Editor/ClothPainter/Private/SClothAssetSelector.h
+++ b/Engine/Source/Editor/ClothPainter/Private/SClothAssetSelector.h
@@ -62,10 +62,6 @@ public:
 
 protected:
 
-#if WITH_APEX_CLOTHING
-	FReply OnImportApexFileClicked();
-#endif
-
 	/* Copies clothing setup from source SkelMesh */
 	void OnCopyClothingAssetSelected(const FAssetData& AssetData);
 
diff --git a/Engine/Source/Editor/ClothingSystemEditor/Private/ClothingAssetFactory.cpp b/Engine/Source/Editor/ClothingSystemEditor/Private/ClothingAssetFactory.cpp
--- a/Engine/Source/Editor/ClothingSystemEditor/Private/ClothingAssetFactory.cpp
+++ b/Engine/Source/Editor/ClothingSystemEditor/Private/ClothingAssetFactory.cpp
@@ -6,11 +6,6 @@
 #include "PhysXPublic.h"
 #include "PhysicsPublic.h"
 #include "ClothingAsset.h"
-#if WITH_APEX_CLOTHING
-#include "ClothingAssetAuthoring.h"
-#include "ApexClothingUtils.h"
-#include "ClothConfigNv.h"
-#endif // WITH_APEX_CLOTHING
 #include "ContentBrowserModule.h"
 #include "Misc/FileHelper.h"
 #include "Engine/SkeletalMesh.h"
@@ -90,111 +85,11 @@ UClothingAssetBase* UClothingAssetFactory::Import
 	FName InName
 )
 {
-#if WITH_APEX_CLOTHING
-	if(!TargetMesh)
-	{
-		return nullptr;
-	}
-
-	UClothingAssetCommon* NewClothingAsset = nullptr;
-
-	TArray<uint8> FileBuffer;
-	if(FFileHelper::LoadFileToArray(FileBuffer, *Filename, FILEREAD_Silent))
-	{
-		ClothingAsset* ApexAsset = ApexClothingUtils::CreateApexClothingAssetFromBuffer(FileBuffer.GetData(), FileBuffer.Num());
-		ApexAsset = ConvertApexAssetCoordSystem(ApexAsset);
-
-		if(InName == NAME_None)
-		{
-			InName = *FPaths::GetBaseFilename(Filename);
-		}
-
-		// Create an unreal clothing asset
-		NewClothingAsset = Cast<UClothingAssetCommon>(CreateFromApexAsset(ApexAsset, TargetMesh, InName));
-
-		if(NewClothingAsset)
-		{
-			// Store import path
-			NewClothingAsset->ImportedFilePath = Filename;
-
-			// Push to the target mesh
-			TargetMesh->AddClothingAsset(NewClothingAsset);
-		}
-	}
-
-	return NewClothingAsset;
-#else
 	return nullptr;
-#endif
 }
 
 UClothingAssetBase* UClothingAssetFactory::Reimport(const FString& Filename, USkeletalMesh* TargetMesh, UClothingAssetBase* OriginalAsset)
 {
-#if WITH_APEX_CLOTHING
-	if(!TargetMesh)
-	{
-		return nullptr;
-	}
-
-	int32 OldIndex = INDEX_NONE;
-	if(TargetMesh->GetMeshClothingAssets().Find(OriginalAsset, OldIndex))
-	{
-		TArray<UActorComponent*> ComponentsToReregister;
-		for(TObjectIterator<USkeletalMeshComponent> It; It; ++It)
-		{
-			if(USkeletalMesh* UsedMesh = (*It)->SkeletalMesh)
-			{
-				if(UsedMesh == TargetMesh)
-				{
-					ComponentsToReregister.Add(*It);
-				}
-			}
-		}
-	
-		FMultiComponentReregisterContext ReregisterContext(ComponentsToReregister);
-
-		UClothingAssetCommon* OldClothingAsset = Cast<UClothingAssetCommon>(TargetMesh->GetMeshClothingAssets()[OldIndex]);
-		UClothingAssetCommon* NewClothingAsset = nullptr;
-		FName AssetName = NAME_None;
-
-		if(!OldClothingAsset || !TargetMesh->GetMeshClothingAssets().IsValidIndex(OldIndex))
-		{
-			return nullptr;
-		}
-
-		TArray<uint8> FileBuffer;
-		if(FFileHelper::LoadFileToArray(FileBuffer, *Filename, FILEREAD_Silent))
-		{
-			ClothingAsset* ApexAsset = ApexClothingUtils::CreateApexClothingAssetFromBuffer(FileBuffer.GetData(), FileBuffer.Num());
-			ApexAsset = ConvertApexAssetCoordSystem(ApexAsset);
-			AssetName = *FPaths::GetBaseFilename(Filename);
-
-			TArray<ClothingAssetUtils::FClothingAssetMeshBinding> AssetBindings;
-			ClothingAssetUtils::GetAllMeshClothingAssetBindings(TargetMesh, AssetBindings);
-
-			OldClothingAsset->UnbindFromSkeletalMesh(TargetMesh);
-
-			// Create an unreal clothing asset
-			NewClothingAsset = Cast<UClothingAssetCommon>(CreateFromApexAsset(ApexAsset, TargetMesh, AssetName));
-
-			if(NewClothingAsset)
-			{
-				// Store import path
-				NewClothingAsset->ImportedFilePath = Filename;
-
-				TargetMesh->GetMeshClothingAssets()[OldIndex] = NewClothingAsset;
-
-				for(ClothingAssetUtils::FClothingAssetMeshBinding& Binding : AssetBindings)
-				{
-					NewClothingAsset->BindToSkeletalMesh(TargetMesh, Binding.LODIndex, Binding.SectionIndex, Binding.AssetInternalLodIndex);
-				}
-			}
-		}
-
-		return NewClothingAsset;
-	}
-#endif
-
 	return nullptr;
 }
 
@@ -380,741 +275,14 @@ UClothingAssetBase* UClothingAssetFactory::ImportLodToClothing(USkeletalMesh* Ta
 
 UClothingAssetBase* UClothingAssetFactory::CreateFromApexAsset(nvidia::apex::ClothingAsset* InApexAsset, USkeletalMesh* TargetMesh, FName InName)
 {
-#if WITH_APEX_CLOTHING
-	UClothingAssetCommon* NewClothingAsset = NewObject<UClothingAssetCommon>(TargetMesh, InName);
-	NewClothingAsset->SetFlags(RF_Transactional);
-	NewClothingAsset->PostUpdateAllAssets();  // Ensure that the required configs, including shared configs, are created
-
-	const NvParameterized::Interface* AssetParams = InApexAsset->getAssetNvParameterized();
-	NvParameterized::Handle GraphicalLodArrayHandle(*AssetParams, "graphicalLods");
-
-	int32 NumSuccessfulLods = 0;
-
-	int32 NumLodsToBuild = 0;
-	GraphicalLodArrayHandle.getArraySize(NumLodsToBuild);
-
-	for(int32 CurrLodIdx = 0; CurrLodIdx < NumLodsToBuild; ++CurrLodIdx)
-	{
-		NewClothingAsset->AddNewLod();
-		FClothLODDataCommon& CurrentLodData = NewClothingAsset->LodData[CurrLodIdx];
-
-		TArray<FApexVertData> ApexVertData;
-
-		ExtractLodPhysicalData(NewClothingAsset, *InApexAsset, CurrLodIdx, CurrentLodData, ApexVertData);
-		ExtractSphereCollisions(NewClothingAsset, *InApexAsset, CurrLodIdx, CurrentLodData);
-		ExtractMaterialParameters(NewClothingAsset, *InApexAsset);
-
-		// Set to use legacy wind calculations, which is what APEX would normally have used
-		UClothConfigNv* const ConfigNv = NewClothingAsset->GetClothConfig<UClothConfigNv>();
-		if (!ensure(ConfigNv))
-		{
-			continue;
-		}
-		ConfigNv->ClothingWindMethod = EClothingWindMethodNv::Legacy;
-
-		// Fixup unreal-side bone indices
-		const int32 NumBoneDatas = CurrentLodData.PhysicalMeshData.BoneData.Num();
-		check(NumBoneDatas == ApexVertData.Num());
-		for(int32 BoneDataIndex = 0; BoneDataIndex < NumBoneDatas; ++BoneDataIndex)
-		{
-			FClothVertBoneData& BoneData = CurrentLodData.PhysicalMeshData.BoneData[BoneDataIndex];
-			FApexVertData& CurrentVertData = ApexVertData[BoneDataIndex];
-		
-			for(int32 BoneInfluenceIdx = 0; BoneInfluenceIdx < MAX_TOTAL_INFLUENCES; ++BoneInfluenceIdx)
-			{
-				uint16 ApexBoneIndex = CurrentVertData.BoneIndices[BoneInfluenceIdx];
-				BoneData.BoneIndices[BoneInfluenceIdx] = ApexBoneIndex;
-			}
-		}
-	}
-
-	ExtractBoneData(NewClothingAsset, *InApexAsset);
-
-	// Now that we've extracted the APEX bone data, we need to fill the generic asset
-	// data with bone data for Unreal rather than APEX internal representations
-	const int32 NumUsedBones = NewClothingAsset->UsedBoneNames.Num();
-
-	NewClothingAsset->UsedBoneIndices.AddDefaulted(NumUsedBones);
-	for(int32 UsedBoneIndex = 0; UsedBoneIndex < NumUsedBones; ++UsedBoneIndex)
-	{
-		const FName& BoneName = NewClothingAsset->UsedBoneNames[UsedBoneIndex];
-		const int32 UnrealBoneIndex = TargetMesh->GetRefSkeleton().FindBoneIndex(BoneName);
-
-		// If we find an invalid bone then the asset is invalid, as it cannot be skinned to this mesh
-		if(UnrealBoneIndex == INDEX_NONE)
-		{
-			FText ErrorText = FText::Format(LOCTEXT("InvalidBoneError", "Imported asset requires bone \"{0}\", which is not present in the skeletal mesh ({1})"), FText::FromName(BoneName), FText::FromString(TargetMesh->GetName()));
-			LogAndToastWarning(ErrorText);
-
-			return nullptr;
-		}
-		
-		NewClothingAsset->UsedBoneIndices[UsedBoneIndex] = UnrealBoneIndex;
-	}
-
-	uint32 AssetInternalRootBoneIndex;
-	verify(NvParameterized::getParamU32(*AssetParams, "rootBoneIndex", AssetInternalRootBoneIndex));
-	FName ConvertedBoneName(*FString(InApexAsset->getBoneName(AssetInternalRootBoneIndex)).Replace(TEXT(" "), TEXT("-")));
-
-	NewClothingAsset->AssetGuid = FGuid::NewGuid();
-	NewClothingAsset->InvalidateCachedData();
-
-	NewClothingAsset->BuildLodTransitionData();
-	NewClothingAsset->CalculateReferenceBoneIndex();
-
-	// Add masks for parameters
-	for (FClothLODDataCommon& LodData : NewClothingAsset->LodData)
-	{
-		const FClothPhysicalMeshData& PhysMesh = LodData.PhysicalMeshData;
-
-		// Didn't do anything previously - clear out in case there's something in there
-		// so we can use it correctly now.
-		LodData.PointWeightMaps.Reset(3);
-
-		// Max distances
-		LodData.PointWeightMaps.AddDefaulted();
-		FPointWeightMap& MaxDistanceMask = LodData.PointWeightMaps.Last();
-		const FPointWeightMap& MaxDistances = PhysMesh.GetWeightMap(EWeightMapTargetCommon::MaxDistance);
-		MaxDistanceMask.Initialize(MaxDistances, EWeightMapTargetCommon::MaxDistance);
-
-		const FPointWeightMap* const BackstopRadiuses = PhysMesh.FindWeightMap(EWeightMapTargetCommon::BackstopRadius);
-		if (BackstopRadiuses && !BackstopRadiuses->IsZeroed())
-		{
-			// Backstop radii
-			LodData.PointWeightMaps.AddDefaulted();
-			FPointWeightMap& BackstopRadiusMask = LodData.PointWeightMaps.Last();
-			BackstopRadiusMask.Initialize(*BackstopRadiuses, EWeightMapTargetCommon::BackstopRadius);
-
-			// Backstop distances
-			LodData.PointWeightMaps.AddDefaulted();
-			FPointWeightMap& BackstopDistanceMask = LodData.PointWeightMaps.Last();
-			const FPointWeightMap& BackstopDistances = PhysMesh.GetWeightMap(EWeightMapTargetCommon::BackstopDistance);
-			BackstopDistanceMask.Initialize(BackstopDistances, EWeightMapTargetCommon::BackstopDistance);
-		}
-	}
-
-	return NewClothingAsset;
-#endif
 	return nullptr;
 }
 
 bool UClothingAssetFactory::CanImport(const FString& Filename)
 {
-#if WITH_APEX_CLOTHING
-	// Need to read in the file and try to create an asset to get it's type
-	TArray<uint8> FileBuffer;
-	if(FFileHelper::LoadFileToArray(FileBuffer, *Filename, FILEREAD_Silent))
-	{
-		physx::PxFileBuf* Stream = GApexSDK->createMemoryReadStream(FileBuffer.GetData(), FileBuffer.Num());
-		if(Stream)
-		{
-			NvParameterized::Serializer::SerializeType SerializeType = GApexSDK->getSerializeType(*Stream);
-			if(NvParameterized::Serializer* Serializer = GApexSDK->createSerializer(SerializeType))
-			{
-				NvParameterized::Serializer::DeserializedData DeserializedData;
-				Serializer->deserialize(*Stream, DeserializedData);
-
-				if(DeserializedData.size() > 0)
-				{
-					NvParameterized::Interface* AssetInterface = DeserializedData[0];
-
-					int32 StringLength = StringCast<TCHAR>(AssetInterface->className()).Length();
-					FString ClassName(StringLength, StringCast<TCHAR>(AssetInterface->className()).Get());
-
-					if(ClassName == ClothingFactoryConstants::ClothingAssetClass)
-					{
-						return true;
-					}
-				}
-			}
-
-			GApexSDK->releaseMemoryReadStream(*Stream);
-		}
-	}
-#endif
-
 	return false;
 }
 
-#if WITH_APEX_CLOTHING
-ClothingAsset* UClothingAssetFactory::ConvertApexAssetCoordSystem(ClothingAsset* InAsset)
-{
-	// Build new asset interface to store the transformed asset
-	const NvParameterized::Interface* OriginalInterface = InAsset->getAssetNvParameterized();
-	NvParameterized::Interface* NewInterface = GApexSDK->getParameterizedTraits()->createNvParameterized(OriginalInterface->className());
-	check(NewInterface);
-
-	// Copy asset data
-	NewInterface->copy(*OriginalInterface);
-
-	ClothingAssetAuthoring* AssetAuthoring = (ClothingAssetAuthoring*)GApexSDK->createAssetAuthoring(NewInterface, nullptr);
-
-	check(AssetAuthoring);
-
-	// Need to check for bone actors and spheres, we can't have both
-	PxI32 NumBoneActors = 0;
-	PxI32 NumBoneSpheres = 0;
-
-	verify(NvParameterized::getParamArraySize(*OriginalInterface, ClothingFactoryConstants::ParamName_BoneActors, NumBoneActors));
-	verify(NvParameterized::getParamArraySize(*OriginalInterface, ClothingFactoryConstants::ParamName_BoneSpheres, NumBoneSpheres));
-
-	// Remove collision if we have spheres and actors (actors will remain)
-	if(NumBoneActors > 0 && NumBoneSpheres > 0)
-	{
-		AssetAuthoring->clearCollision();
-	}
-
-	// Y direction needs to be inverted
-	PxMat44 YInvertMatrix(PxIdentity);
-	YInvertMatrix.column1.y = -1.0f;
-
-	// Matrix holding the coordinate space conversion required for the mesh
-	PxMat44 ConversionTransform(PxIdentity);
-
-	// Get gravity direction, as that should be -up
-	PxVec3 GravityDirection;
-	verify(NvParameterized::getParamVec3(*OriginalInterface, ClothingFactoryConstants::ParamName_GravityDirection, GravityDirection));
-
-	// Y-up, needs conversion to z-up
-	if(GravityDirection.z == 0.0f && FMath::Abs(GravityDirection.y) > 0.0f)
-	{
-		PxVec3 NewGravityDirection(GravityDirection.x, GravityDirection.z, GravityDirection.y);
-		AssetAuthoring->setSimulationGravityDirection(NewGravityDirection);
-
-		// Invert Y + 90 deg rotation on x
-		ConversionTransform.column1.y = 0.0f;
-		ConversionTransform.column1.z = 1.0f;
-		ConversionTransform.column2.y = 1.0f;
-		ConversionTransform.column2.z = 0.0f;
-	}
-	else
-	{
-		ConversionTransform = YInvertMatrix;
-	}
-
-	AssetAuthoring->applyTransformation(ConversionTransform, 1.0f, true, true);
-
-	// Transform bind poses
-	const uint32 NumUsedBones = InAsset->getNumUsedBones();
-
-	TArray<PxMat44> TransformedBindPoses;
-	TransformedBindPoses.Reserve(NumUsedBones);
-
-	for(uint32 Idx = 0; Idx < NumUsedBones; ++Idx)
-	{
-		PxMat44 CurrentBindPose(PxIdentity);
-		AssetAuthoring->getBoneBindPose(Idx, CurrentBindPose);
-		TransformedBindPoses.Add(CurrentBindPose * YInvertMatrix);
-	}
-
-	AssetAuthoring->updateBindPoses(TransformedBindPoses.GetData(), TransformedBindPoses.Num(), true, true);
-
-	const uint32 NumLods = AssetAuthoring->getNumLods();
-	for(uint32 Idx = 0; Idx < NumLods; ++Idx)
-	{
-		if(NvParameterized::Interface* RenderMeshAuthoringInterface = AssetAuthoring->getRenderMeshAssetAuthoring(Idx))
-		{
-			NvParameterized::Handle RenderMeshAuthoringHandle(RenderMeshAuthoringInterface);
-
-			bool bFlipU = false;
-			bool bFlipV = false;
-
-			NvParameterized::Interface* UVOriginParameter = NvParameterized::findParam(*RenderMeshAuthoringInterface, ClothingFactoryConstants::ParamName_UvOrigin, RenderMeshAuthoringHandle);
-			if(UVOriginParameter)
-			{
-				uint32 UvOrigin = 0;
-				RenderMeshAuthoringHandle.getParamU32(UvOrigin);
-
-				switch(UvOrigin)
-				{
-				case TextureUVOrigin::ORIGIN_TOP_LEFT:
-					bFlipU = false;
-					bFlipV = false;
-					break;
-				case TextureUVOrigin::ORIGIN_TOP_RIGHT:
-					bFlipU = true;
-					bFlipV = false;
-					break;
-				case TextureUVOrigin::ORIGIN_BOTTOM_LEFT:
-					bFlipU = false;
-					bFlipV = true;
-					break;
-				case TextureUVOrigin::ORIGIN_BOTTOM_RIGHT:
-					bFlipU = false;
-					bFlipV = false;
-					break;
-				default:
-					break;
-				}
-
-				RenderMeshAuthoringHandle.setParamU32(TextureUVOrigin::ORIGIN_TOP_LEFT);
-			}
-
-			// Flip UVs
-			FlipAuthoringUvs(RenderMeshAuthoringInterface, bFlipU, bFlipV);
-		}
-		else
-		{
-			break;
-		}
-	}
-
-	char AssetName[MAX_SPRINTF] = {0};
-	FCStringAnsi::Strncpy(AssetName, InAsset->getName(), MAX_SPRINTF);
-
-	InAsset->release();
-	InAsset = nullptr;
-
-	ClothingAsset* NewAsset = (ClothingAsset*)GApexSDK->createAsset(*AssetAuthoring, AssetName);
-
-	check(NewAsset);
-
-	AssetAuthoring->release();
-
-	return NewAsset;
-}
-
-void UClothingAssetFactory::FlipAuthoringUvs(NvParameterized::Interface* InRenderMeshAuthoringInterface, bool bFlipU, bool bFlipV)
-{
-	if(!bFlipU && !bFlipV)
-	{
-		// Don't need to do anything
-		return;
-	}
-
-	NvParameterized::Handle SubmeshArrayHandle(*InRenderMeshAuthoringInterface, ClothingFactoryConstants::ParamName_SubmeshArray);
-
-	if(SubmeshArrayHandle.isValid())
-	{
-		int32 ArraySize = 0;
-		SubmeshArrayHandle.getArraySize(ArraySize, 0);
-
-		for(int32 SubmeshIdx = 0; SubmeshIdx < ArraySize; ++SubmeshIdx)
-		{
-			NvParameterized::Handle SubmeshHandle(SubmeshArrayHandle);
-			SubmeshArrayHandle.getChildHandle(SubmeshIdx, SubmeshHandle);
-
-			if(!SubmeshHandle.isValid())
-			{
-				// No submesh, move to next array entry
-				continue;
-			}
-
-			NvParameterized::Interface* SubmeshInterface = nullptr;
-			SubmeshHandle.getParamRef(SubmeshInterface);
-
-			NvParameterized::Handle BufferFormatsHandle(SubmeshHandle);
-			NvParameterized::findParam(*SubmeshInterface, ClothingFactoryConstants::ParamName_SubmeshBufferFormats, BufferFormatsHandle);
-
-			if(!BufferFormatsHandle.isValid())
-			{
-				// No valid format array, move to next submesh
-				continue;
-			}
-
-			int32 FormatArraySize = 0;
-			BufferFormatsHandle.getArraySize(FormatArraySize);
-
-			for(int32 FormatIdx = 0; FormatIdx < FormatArraySize; ++FormatIdx)
-			{
-				NvParameterized::Handle FormatHandle(BufferFormatsHandle);
-				FormatHandle.set(FormatIdx);
-
-				NvParameterized::Handle SemanticHandle(FormatHandle);
-				FormatHandle.getChildHandle(FormatHandle.getInterface(), "semantic", SemanticHandle);
-
-				if(!SemanticHandle.isValid())
-				{
-					// No valid semantic, move to next buffer format
-					continue;
-				}
-
-				PxI32 BufferSemantic = -1;
-				SemanticHandle.getParamI32(BufferSemantic);
-
-				if(BufferSemantic >= RenderVertexSemantic::TEXCOORD0 &&
-				   BufferSemantic <= RenderVertexSemantic::TEXCOORD3)
-				{
-					NvParameterized::Handle BufferArrayHandle(SubmeshHandle);
-					NvParameterized::findParam(*SubmeshInterface, ClothingFactoryConstants::ParamName_VertexBuffers, BufferArrayHandle);
-
-					int32 BufferArraySize = -1;
-					BufferArrayHandle.getArraySize(BufferArraySize);
-					check(BufferSemantic < BufferArraySize);
-
-					if(BufferArraySize == -1)
-					{
-						// Failed to find array, move to next format
-						continue;
-					}
-
-					BufferArrayHandle.set(BufferSemantic);
-
-					NvParameterized::Handle DataHandle(BufferArrayHandle);
-					NvParameterized::Interface* BufferInterface = nullptr;
-					BufferArrayHandle.getParamRef(BufferInterface);
-
-					check(BufferInterface);
-
-					NvParameterized::findParam(*BufferInterface, ClothingFactoryConstants::ParamName_BufferData, DataHandle);
-
-					if(!DataHandle.isValid())
-					{
-						// No data array, move to next format
-						continue;
-					}
-
-					int32 DataArraySize = -1;
-					DataHandle.getArraySize(DataArraySize);
-
-					float MaxU = -MAX_flt;
-					float MaxV = -MAX_flt;
-					for(int32 DataIdx = 0; DataIdx < DataArraySize; ++DataIdx)
-					{
-						// Push to data entry
-						DataHandle.set(DataIdx);
-
-						// UV coord storage
-						float Coord[2];
-
-						DataHandle.set(0); // Inside data entry, get first element (U Coord)
-						DataHandle.getParamF32(Coord[0]);
-						DataHandle.popIndex(); // Back out to data entry
-
-						DataHandle.set(1); // Inside data entry, get second element (V Coord)
-						DataHandle.getParamF32(Coord[1]);
-						DataHandle.popIndex(); // Back out to data entry
-
-						MaxU = FMath::Max(MaxU, Coord[0] - DELTA);
-						MaxV = FMath::Max(MaxV, Coord[1] - DELTA);
-
-						DataHandle.popIndex(); // Back out to data array
-					}
-
-					MaxU = FMath::FloorToFloat(MaxU) + 1.0f;
-					MaxV = FMath::FloorToFloat(MaxV) + 1.0f;
-
-					for(int32 DataIdx = 0; DataIdx < DataArraySize; ++DataIdx)
-					{
-						DataHandle.set(DataIdx);
-
-						float CoordPart = 0.0f;
-
-						if(bFlipU)
-						{
-							DataHandle.set(0);
-							DataHandle.getParamF32(CoordPart);
-							DataHandle.setParamF32(MaxU - CoordPart);
-							DataHandle.popIndex();
-						}
-
-						if(bFlipV)
-						{
-							DataHandle.set(1);
-							DataHandle.getParamF32(CoordPart);
-							DataHandle.setParamF32(MaxV - CoordPart);
-							DataHandle.popIndex();
-						}
-
-						DataHandle.popIndex();
-					}
-				}
-			}
-		}
-	}
-}
-
-void UClothingAssetFactory::ExtractBoneData(UClothingAssetCommon* NewAsset, ClothingAsset &InApexAsset)
-{
-	const uint32 NumApexUsedBones = InApexAsset.getNumUsedBones();
-
-	NewAsset->UsedBoneNames.Empty(NumApexUsedBones);
-
-	for(uint32 BoneIdx = 0; BoneIdx < NumApexUsedBones; ++BoneIdx)
-	{
-		FString BoneName = FString(InApexAsset.getBoneName(BoneIdx)).Replace(TEXT(" "), TEXT("-"));
-		NewAsset->UsedBoneNames.Add(*BoneName);
-	}
-}
-
-void UClothingAssetFactory::ExtractSphereCollisions(UClothingAssetCommon* NewAsset, nvidia::apex::ClothingAsset &InApexAsset, int32 InLodIdx, FClothLODDataCommon& InLodData)
-{
-	const NvParameterized::Interface* AssetParams = InApexAsset.getAssetNvParameterized();
-
-	NvParameterized::Handle BoneSphereHandle(*AssetParams, "boneSpheres");
-
-	int32 NumBoneSpheres = 0;
-	BoneSphereHandle.getArraySize(NumBoneSpheres);
-
-	FClothCollisionData& CollisionData = InLodData.CollisionData;
-	CollisionData.Spheres.AddDefaulted(NumBoneSpheres);
-
-	// Load the bone spheres
-	for(int32 BoneSphereIndex = 0; BoneSphereIndex < NumBoneSpheres; ++BoneSphereIndex)
-	{
-		FClothCollisionPrim_Sphere& CurrentSphere = CollisionData.Spheres[BoneSphereIndex];
-
-		BoneSphereHandle.set(BoneSphereIndex);
-
-		NvParameterized::Handle ChildHandle(BoneSphereHandle);
-
-		BoneSphereHandle.getChildHandle(BoneSphereHandle.getInterface(), "boneIndex", ChildHandle);
-		ChildHandle.getParamI32(CurrentSphere.BoneIndex);
-
-		BoneSphereHandle.getChildHandle(BoneSphereHandle.getInterface(), "radius", ChildHandle);
-		ChildHandle.getParamF32(CurrentSphere.Radius);
-
-		BoneSphereHandle.getChildHandle(BoneSphereHandle.getInterface(), "localPos", ChildHandle);
-		physx::PxVec3 PxLocalPos;
-		ChildHandle.getParamVec3(PxLocalPos);
-		CurrentSphere.LocalPosition = P2UVector(PxLocalPos);
-
-		BoneSphereHandle.popIndex();
-	}
-
-	// Next load "connections". A connection is used to turn 2 spheres into a capsule by connecting them
-	NvParameterized::Handle BoneSphereConnectionHandle(*AssetParams, "boneSphereConnections");
-
-	int32 NumConnections = 0;
-	BoneSphereConnectionHandle.getArraySize(NumConnections);
-	check(NumConnections % 2 == 0); // Needs to be even
-	CollisionData.SphereConnections.AddDefaulted(NumConnections / 2);
-
-	for(int32 ConnectionIndex = 0; ConnectionIndex < NumConnections; ConnectionIndex += 2)
-	{
-		FClothCollisionPrim_SphereConnection& CurrentConnection = CollisionData.SphereConnections[ConnectionIndex / 2];
-
-		uint16 FirstSphereIndex;
-		uint16 SecondSphereIndex;
-
-		BoneSphereConnectionHandle.set(ConnectionIndex);
-		BoneSphereConnectionHandle.getParamU16(FirstSphereIndex);
-		BoneSphereConnectionHandle.popIndex();
-
-		BoneSphereConnectionHandle.set(ConnectionIndex + 1);
-		BoneSphereConnectionHandle.getParamU16(SecondSphereIndex);
-		BoneSphereConnectionHandle.popIndex();
-
-		CurrentConnection.SphereIndices[0] = FirstSphereIndex;
-		CurrentConnection.SphereIndices[1] = SecondSphereIndex;
-	}
-
-	// Load boneactors. Bone actors are a different way to handle capsules
-	// By defining a capsule height and radius.
-	NvParameterized::Handle BoneActorHandle(*AssetParams, "boneActors");
-
-	int32 NumActors = 0;
-	BoneActorHandle.getArraySize(NumActors);
-	
-	for(int32 ActorIndex = 0; ActorIndex < NumActors; ++ActorIndex)
-	{
-		BoneActorHandle.set(ActorIndex);
-
-		NvParameterized::Handle ChildHandle(BoneActorHandle);
-		BoneActorHandle.getChildHandle(BoneActorHandle.getInterface(), "convexVerticesCount", ChildHandle);
-
-		uint32 NumConvexVerts = 0;
-		ChildHandle.getParamU32(NumConvexVerts);
-
-		if(NumConvexVerts > 0)
-		{
-			// Convex mesh, extract the data
-		}
-		else
-		{
-			int32 BoneIndex;
-			float Radius;
-			float Height;
-			physx::PxMat44 PxPoseMatrix;
-			FMatrix PoseMatrix;
-
-			BoneActorHandle.getChildHandle(BoneActorHandle.getInterface(), "boneIndex", ChildHandle);
-			ChildHandle.getParamI32(BoneIndex);
-
-			BoneActorHandle.getChildHandle(BoneActorHandle.getInterface(), "capsuleRadius", ChildHandle);
-			ChildHandle.getParamF32(Radius);
-
-			BoneActorHandle.getChildHandle(BoneActorHandle.getInterface(), "capsuleHeight", ChildHandle);
-			ChildHandle.getParamF32(Height);
-
-			BoneActorHandle.getChildHandle(BoneActorHandle.getInterface(), "localPose", ChildHandle);
-			ChildHandle.getParamMat44(PxPoseMatrix);
-
-			PoseMatrix = P2UMatrix(PxPoseMatrix);
-
-			FVector HalfVector(0.0f, Height * 0.5f, 0.0f);
-			FVector Sphere0Position = PoseMatrix.TransformPosition(HalfVector);
-			FVector Sphere1Position = PoseMatrix.TransformPosition(-HalfVector);
-
-			CollisionData.Spheres.AddDefaulted(2);
-			FClothCollisionPrim_Sphere& Sphere0 = CollisionData.Spheres.Last(1);
-			FClothCollisionPrim_Sphere& Sphere1 = CollisionData.Spheres.Last(0);
-
-			Sphere0.LocalPosition = Sphere0Position;
-			Sphere0.Radius = Radius;
-			Sphere0.BoneIndex = BoneIndex;
-
-			Sphere1.LocalPosition = Sphere1Position;
-			Sphere1.Radius = Radius;
-			Sphere1.BoneIndex = BoneIndex;
-
-			CollisionData.SphereConnections.AddDefaulted();
-			FClothCollisionPrim_SphereConnection& Connection = CollisionData.SphereConnections.Last();
-
-			Connection.SphereIndices[0] = CollisionData.Spheres.Num() - 2;
-			Connection.SphereIndices[1] = CollisionData.Spheres.Num() - 1;
-		}
-
-		BoneActorHandle.popIndex();
-	}
-}
-
-void UClothingAssetFactory::ExtractMaterialParameters(UClothingAssetCommon* NewAsset, nvidia::apex::ClothingAsset &InApexAsset)
-{
-	const NvParameterized::Interface* AssetParams = InApexAsset.getAssetNvParameterized();
-
-	uint32 MaterialIndex = INDEX_NONE;
-	NvParameterized::getParamU32(*AssetParams, "materialIndex", MaterialIndex);
-
-	NvParameterized::Interface* MaterialLibraryParams = nullptr;
-	NvParameterized::getParamRef(*AssetParams, "materialLibrary", MaterialLibraryParams);
-
-	NvParameterized::Handle MaterialArrayHandle(*MaterialLibraryParams);
-	MaterialLibraryParams->getParameterHandle("materials", MaterialArrayHandle);
-	
-	int32 NumMaterials = INDEX_NONE;
-	MaterialArrayHandle.getArraySize(NumMaterials);
-
-	check(MaterialIndex < (uint32)NumMaterials);
-
-	MaterialArrayHandle.set(MaterialIndex);
-
-	if(UClothConfigNv* const ConfigPtr = NewAsset->GetClothConfig<UClothConfigNv>())
-	{
-		UClothConfigNv& Config = *ConfigPtr;
-		NvParameterized::Handle ChildHandle(MaterialArrayHandle);
-
-		// Read out material params
-		MaterialArrayHandle.getChildHandle(MaterialArrayHandle.getInterface(), "selfcollisionThickness", ChildHandle);
-		ChildHandle.getParamF32(Config.SelfCollisionRadius);
-
-		MaterialArrayHandle.getChildHandle(MaterialArrayHandle.getInterface(), "selfcollisionStiffness", ChildHandle);
-		ChildHandle.getParamF32(Config.SelfCollisionStiffness);
-
-		float ApexDamping;
-		MaterialArrayHandle.getChildHandle(MaterialArrayHandle.getInterface(), "damping", ChildHandle);
-		ChildHandle.getParamF32(ApexDamping);
-		Config.Damping = FVector(ApexDamping);
-
-		MaterialArrayHandle.getChildHandle(MaterialArrayHandle.getInterface(), "friction", ChildHandle);
-		ChildHandle.getParamF32(Config.Friction);
-
-		MaterialArrayHandle.getChildHandle(MaterialArrayHandle.getInterface(), "gravityScale", ChildHandle);
-		ChildHandle.getParamF32(Config.GravityScale);
-
-		// Tether parameters
-
-		MaterialArrayHandle.getChildHandle(MaterialArrayHandle.getInterface(), "tetherLimit", ChildHandle);
-		ChildHandle.getParamF32(Config.TetherLimit);
-
-		MaterialArrayHandle.getChildHandle(MaterialArrayHandle.getInterface(), "tetherStiffness", ChildHandle);
-		ChildHandle.getParamF32(Config.TetherStiffness);
-
-		// Drag and inertia have 2 components but APEX only uses one
-		float Drag = 1.0f;
-		MaterialArrayHandle.getChildHandle(MaterialArrayHandle.getInterface(), "drag", ChildHandle);
-		ChildHandle.getParamF32(Drag);
-		Config.LinearDrag = FVector(Drag);
-		Config.AngularDrag = FVector(Drag);
-
-		float InertiaScale = 1.0f;
-		MaterialArrayHandle.getChildHandle(MaterialArrayHandle.getInterface(), "inertiaScale", ChildHandle);
-		ChildHandle.getParamF32(InertiaScale);
-		Config.LinearInertiaScale = FVector(InertiaScale);
-		Config.AngularInertiaScale = FVector(InertiaScale);
-
-		// Simulation frequencies
-		MaterialArrayHandle.getChildHandle(MaterialArrayHandle.getInterface(), "stiffnessFrequency", ChildHandle);
-		ChildHandle.getParamF32(Config.StiffnessFrequency);
-
-		MaterialArrayHandle.getChildHandle(MaterialArrayHandle.getInterface(), "solverFrequency", ChildHandle);
-		ChildHandle.getParamF32(Config.SolverFrequency);
-
-		// Vertical constraint params
-
-		MaterialArrayHandle.getChildHandle(MaterialArrayHandle.getInterface(), "verticalStretchingStiffness", ChildHandle);
-		ChildHandle.getParamF32(Config.VerticalConstraint.Stiffness);
-
-		MaterialArrayHandle.getChildHandle(MaterialArrayHandle.getInterface(), "verticalStiffnessScaling.compressionRange", ChildHandle);
-		ChildHandle.getParamF32(Config.VerticalConstraint.CompressionLimit);
-
-		MaterialArrayHandle.getChildHandle(MaterialArrayHandle.getInterface(), "verticalStiffnessScaling.stretchRange", ChildHandle);
-		ChildHandle.getParamF32(Config.VerticalConstraint.StretchLimit);
-
-		MaterialArrayHandle.getChildHandle(MaterialArrayHandle.getInterface(), "verticalStiffnessScaling.scale", ChildHandle);
-		ChildHandle.getParamF32(Config.VerticalConstraint.StiffnessMultiplier);
-
-		// Horizontal constraint params
-
-		MaterialArrayHandle.getChildHandle(MaterialArrayHandle.getInterface(), "horizontalStretchingStiffness", ChildHandle);
-		ChildHandle.getParamF32(Config.HorizontalConstraint.Stiffness);
-
-		MaterialArrayHandle.getChildHandle(MaterialArrayHandle.getInterface(), "horizontalStiffnessScaling.compressionRange", ChildHandle);
-		ChildHandle.getParamF32(Config.HorizontalConstraint.CompressionLimit);
-
-		MaterialArrayHandle.getChildHandle(MaterialArrayHandle.getInterface(), "horizontalStiffnessScaling.stretchRange", ChildHandle);
-		ChildHandle.getParamF32(Config.HorizontalConstraint.StretchLimit);
-
-		MaterialArrayHandle.getChildHandle(MaterialArrayHandle.getInterface(), "horizontalStiffnessScaling.scale", ChildHandle);
-		ChildHandle.getParamF32(Config.HorizontalConstraint.StiffnessMultiplier);
-
-		// Bend constraint params
-
-		MaterialArrayHandle.getChildHandle(MaterialArrayHandle.getInterface(), "bendingStiffness", ChildHandle);
-		ChildHandle.getParamF32(Config.BendConstraint.Stiffness);
-
-		MaterialArrayHandle.getChildHandle(MaterialArrayHandle.getInterface(), "bendingStiffnessScaling.compressionRange", ChildHandle);
-		ChildHandle.getParamF32(Config.BendConstraint.CompressionLimit);
-
-		MaterialArrayHandle.getChildHandle(MaterialArrayHandle.getInterface(), "bendingStiffnessScaling.stretchRange", ChildHandle);
-		ChildHandle.getParamF32(Config.BendConstraint.StretchLimit);
-
-		MaterialArrayHandle.getChildHandle(MaterialArrayHandle.getInterface(), "bendingStiffnessScaling.scale", ChildHandle);
-		ChildHandle.getParamF32(Config.BendConstraint.StiffnessMultiplier);
-
-		// Shear constraint params
-
-		MaterialArrayHandle.getChildHandle(MaterialArrayHandle.getInterface(), "shearingStiffness", ChildHandle);
-		ChildHandle.getParamF32(Config.ShearConstraint.Stiffness);
-
-		MaterialArrayHandle.getChildHandle(MaterialArrayHandle.getInterface(), "shearingStiffnessScaling.compressionRange", ChildHandle);
-		ChildHandle.getParamF32(Config.ShearConstraint.CompressionLimit);
-
-		MaterialArrayHandle.getChildHandle(MaterialArrayHandle.getInterface(), "shearingStiffnessScaling.stretchRange", ChildHandle);
-		ChildHandle.getParamF32(Config.ShearConstraint.StretchLimit);
-
-		MaterialArrayHandle.getChildHandle(MaterialArrayHandle.getInterface(), "shearingStiffnessScaling.scale", ChildHandle);
-		ChildHandle.getParamF32(Config.ShearConstraint.StiffnessMultiplier);
-
-		// UE just used the vertical config for everything, so stomp the other configs
-		Config.HorizontalConstraint.CompressionLimit = Config.VerticalConstraint.CompressionLimit;
-		Config.HorizontalConstraint.StretchLimit = Config.VerticalConstraint.StretchLimit;
-		Config.HorizontalConstraint.StiffnessMultiplier = Config.VerticalConstraint.StiffnessMultiplier;
-
-		Config.BendConstraint.CompressionLimit = Config.VerticalConstraint.CompressionLimit;
-		Config.BendConstraint.StretchLimit = Config.VerticalConstraint.StretchLimit;
-		Config.BendConstraint.StiffnessMultiplier = Config.VerticalConstraint.StiffnessMultiplier;
-
-		Config.ShearConstraint.CompressionLimit = Config.VerticalConstraint.CompressionLimit;
-		Config.ShearConstraint.StretchLimit = Config.VerticalConstraint.StretchLimit;
-		Config.ShearConstraint.StiffnessMultiplier = Config.VerticalConstraint.StiffnessMultiplier;
-
-	}
-}
-
-#endif // WITH_APEX_CLOTHING
-
-
 bool UClothingAssetFactory::ImportToLodInternal(
 	USkeletalMesh* SourceMesh, 
 	int32 SourceLodIndex, 
@@ -1327,222 +495,4 @@ bool UClothingAssetFactory::ImportToLodInternal(
 	return true;
 }
 
-#if WITH_APEX_CLOTHING
-
-void UClothingAssetFactory::ExtractLodPhysicalData(UClothingAssetCommon* NewAsset, ClothingAsset &InApexAsset, int32 InLodIdx, FClothLODDataCommon& InLodData, TArray<FApexVertData>& OutApexVertData)
-{
-	const NvParameterized::Interface* AssetParams = InApexAsset.getAssetNvParameterized();
-
-	FClothPhysicalMeshData& PhysData = InLodData.PhysicalMeshData;
-	NvParameterized::Handle GraphicalMeshArrayHandle(*AssetParams, "graphicalLods");
-
-	int32 NumGraphicalLods = 0;
-	GraphicalMeshArrayHandle.getArraySize(NumGraphicalLods);
-
-	uint32 PhysicalMeshIndex = INDEX_NONE;
-
-	for(int32 GraphicalMeshIndex = 0; GraphicalMeshIndex < NumGraphicalLods; ++GraphicalMeshIndex)
-	{
-		NvParameterized::Handle GraphicalMeshHandle(GraphicalMeshArrayHandle);
-		GraphicalMeshArrayHandle.getChildHandle(GraphicalMeshIndex, GraphicalMeshHandle);
-
-		NvParameterized::Interface* MeshInterface = nullptr;
-		GraphicalMeshHandle.getParamRef(MeshInterface);
-
-		NvParameterized::Handle MeshPropertyHandle(MeshInterface);
-
-		MeshInterface->getParameterHandle("lod", MeshPropertyHandle);
-
-		uint32 MeshLodIndex = INDEX_NONE;
-		MeshPropertyHandle.getParamU32(MeshLodIndex);
-
-		if(MeshLodIndex == InLodIdx)
-		{
-			// This is the LOD we want
-			MeshInterface->getParameterHandle("physicalMeshId", MeshPropertyHandle);
-			MeshPropertyHandle.getParamU32(PhysicalMeshIndex);
-		}
-	}
-
-	check(PhysicalMeshIndex != INDEX_NONE);
-
-	NvParameterized::Handle PhysicalMeshArrayHandle(*AssetParams, "physicalMeshes");
-
-	int32 NumPhysicalMeshes = 0;
-	PhysicalMeshArrayHandle.getArraySize(NumPhysicalMeshes);
-
-	check(PhysicalMeshIndex < (uint32)NumPhysicalMeshes)
-
-	{
-		NvParameterized::Handle PhysMeshHandle(PhysicalMeshArrayHandle);
-		PhysicalMeshArrayHandle.getChildHandle(PhysicalMeshIndex, PhysMeshHandle);
-
-		NvParameterized::Interface* PhysicalMeshRef = nullptr;
-		PhysMeshHandle.getParamRef(PhysicalMeshRef);
-
-		NvParameterized::Handle TempHandle(PhysicalMeshRef);
-
-		uint32 NumVertices = 0;
-		uint32 NumIndices = 0;
-		TempHandle.getParameter("physicalMesh.numVertices");
-		TempHandle.getParam(NumVertices);
-
-		TempHandle.getParameter("physicalMesh.numIndices");
-		TempHandle.getParam(NumIndices);
-
-		PhysData.Vertices.Empty();
-		PhysData.Normals.Empty();
-		PhysData.Vertices.AddUninitialized(NumVertices);
-		PhysData.Normals.AddUninitialized(NumVertices);
-
-		// Extract verts
-		TempHandle.getParameter("physicalMesh.vertices");
-		int32 VertArraySize = 0;
-		TempHandle.getArraySize(VertArraySize);
-
-		check(VertArraySize == NumVertices);
-
-		NvParameterized::Handle IterHandle(TempHandle);
-		for(int32 Idx = 0; Idx < VertArraySize; ++Idx)
-		{
-			TempHandle.getChildHandle(Idx, IterHandle);
-
-			PxVec3 PxPosition;
-			IterHandle.getParamVec3(PxPosition);
-
-			PhysData.Vertices[Idx] = P2UVector(PxPosition);
-		}
-
-		// Extract normals
-		TempHandle.getParameter("physicalMesh.normals");
-		int32 NormalArraySize = 0;
-		TempHandle.getArraySize(NormalArraySize);
-
-		check(NormalArraySize == NumVertices);
-
-		IterHandle = NvParameterized::Handle(TempHandle);
-		for(int32 Idx = 0; Idx < VertArraySize; ++Idx)
-		{
-			TempHandle.getChildHandle(Idx, IterHandle);
-
-			PxVec3 PxNormal;
-			IterHandle.getParamVec3(PxNormal);
-
-			PhysData.Normals[Idx] = P2UVector(PxNormal);
-		}
-
-		// Extract indices
-		TempHandle.getParameter("physicalMesh.indices");
-		int32 IndexArraySize = 0;
-		TempHandle.getArraySize(IndexArraySize);
-
-		PhysData.Indices.AddZeroed(IndexArraySize);
-
-		IterHandle = NvParameterized::Handle(TempHandle);
-		for(int32 Idx = 0; Idx < IndexArraySize; ++Idx)
-		{
-			TempHandle.getChildHandle(Idx, IterHandle);
-
-			PxU32 Index;
-			IterHandle.getParamU32(Index);
-
-			PhysData.Indices[Idx] = Index;
-		}
-
-		// Bone data
-		NvParameterized::Handle IndexHandle(TempHandle);
-
-		TempHandle.getParameter("physicalMesh.boneWeights");
-		IndexHandle.getParameter("physicalMesh.boneIndices");
-
-		int32 BoneWeightArraySize = 0;
-		TempHandle.getArraySize(BoneWeightArraySize);
-
-		if(BoneWeightArraySize > 0)
-		{
-			int32 BoneIndexArraySize = 0;
-			IndexHandle.getArraySize(BoneIndexArraySize);
-			check(BoneIndexArraySize == BoneWeightArraySize);
-
-			int32 MaxWeights = BoneWeightArraySize / PhysData.Vertices.Num();
-
-			PhysData.MaxBoneWeights = MaxWeights;
-			PhysData.BoneData.AddZeroed(PhysData.Vertices.Num());
-
-			// Allocate apex only data
-			OutApexVertData.AddDefaulted(PhysData.Vertices.Num());
-
-			NvParameterized::Handle WeightChildHandle(TempHandle);
-			NvParameterized::Handle IndexChildHandle(IndexHandle);
-
-			for(int32 WeightIdx = 0; WeightIdx < BoneWeightArraySize; ++WeightIdx)
-			{
-				TempHandle.getChildHandle(WeightIdx, WeightChildHandle);
-				IndexHandle.getChildHandle(WeightIdx, IndexChildHandle);
-
-				int32 VertIdx = WeightIdx / MaxWeights;
-				int32 VertWeightIdx = WeightIdx % MaxWeights;
-
-				if(VertWeightIdx < MAX_TOTAL_INFLUENCES)
-				{
-					WeightChildHandle.getParamF32(PhysData.BoneData[VertIdx].BoneWeights[VertWeightIdx]);
-					IndexChildHandle.getParamU16(OutApexVertData[VertIdx].BoneIndices[VertWeightIdx]);
-				}
-				else
-				{
-					UE_LOG(LogClothingAssetFactory, Warning, TEXT("Warning, encountered a bone influence index greater than %d, skipping this influence."), MAX_TOTAL_INFLUENCES);
-				}
-			}
-		}
-		
-		// Extract max distances and backstops
-		TempHandle.getParameter("physicalMesh.constrainCoefficients");
-		int32 CoeffArraySize = 0;
-		TempHandle.getArraySize(CoeffArraySize);
-
-		check(CoeffArraySize == NumVertices);
-
-		TArray<float>& MaxDistances = PhysData.FindOrAddWeightMap(EWeightMapTargetCommon::MaxDistance).Values;
-		TArray<float>& BackstopDistances = PhysData.FindOrAddWeightMap(EWeightMapTargetCommon::BackstopDistance).Values;
-		TArray<float>& BackstopRadiuses = PhysData.FindOrAddWeightMap(EWeightMapTargetCommon::BackstopRadius).Values;
-		MaxDistances.Reset(CoeffArraySize);
-		MaxDistances.AddZeroed(CoeffArraySize);
-		BackstopDistances.Reset(CoeffArraySize);
-		BackstopDistances.AddZeroed(CoeffArraySize);
-		BackstopRadiuses.Reset(CoeffArraySize);
-		BackstopRadiuses.AddZeroed(CoeffArraySize);
-
-		IterHandle = NvParameterized::Handle(TempHandle);
-		NvParameterized::Handle ChildHandle(PhysicalMeshRef);
-		for(int32 Idx = 0; Idx < CoeffArraySize; ++Idx)
-		{
-			TempHandle.getChildHandle(Idx, IterHandle);
-
-			IterHandle.getChildHandle(PhysicalMeshRef, "maxDistance", ChildHandle);
-			ChildHandle.getParamF32(MaxDistances[Idx]);
-
-			IterHandle.getChildHandle(AssetParams, "collisionSphereDistance", ChildHandle);
-			ChildHandle.getParamF32(BackstopDistances[Idx]);
-
-			IterHandle.getChildHandle(AssetParams, "collisionSphereRadius", ChildHandle);
-			ChildHandle.getParamF32(BackstopRadiuses[Idx]);
-
-			BackstopDistances[Idx] += BackstopRadiuses[Idx];
-		}
-
-		// Calculate how many fixed verts we have
-		PhysData.NumFixedVerts = 0;
-		for(float Distance : MaxDistances)
-		{
-			if(Distance == 0.0f)
-			{
-				++PhysData.NumFixedVerts;
-			}
-		}
-
-		UE_LOG(LogClothingAssetFactory, Log, TEXT("Finished physical mesh import"));
-	}
-}
-#endif // WITH_APEX_CLOTHING
-
 #undef LOCTEXT_NAMESPACE
diff --git a/Engine/Source/Editor/ClothingSystemEditor/Private/ClothingSystemEditorModule.cpp b/Engine/Source/Editor/ClothingSystemEditor/Private/ClothingSystemEditorModule.cpp
--- a/Engine/Source/Editor/ClothingSystemEditor/Private/ClothingSystemEditorModule.cpp
+++ b/Engine/Source/Editor/ClothingSystemEditor/Private/ClothingSystemEditorModule.cpp
@@ -19,19 +19,11 @@ FClothingSystemEditorModule::FClothingSystemEditorModule()
 void FClothingSystemEditorModule::StartupModule()
 {
 	IModularFeatures::Get().RegisterModularFeature(IClothingAssetFactoryProvider::FeatureName, this);
-
-#if WITH_NVCLOTH
-	IModularFeatures::Get().RegisterModularFeature(FClothingSystemEditorInterfaceModule::ExtenderFeatureName, &NvEditorExtender);
-#endif
 }
 
 void FClothingSystemEditorModule::ShutdownModule()
 {
 	IModularFeatures::Get().UnregisterModularFeature(IClothingAssetFactoryProvider::FeatureName, this);
-
-#if WITH_NVCLOTH
-	IModularFeatures::Get().UnregisterModularFeature(FClothingSystemEditorInterfaceModule::ExtenderFeatureName, &NvEditorExtender);
-#endif
 }
 
 UClothingAssetFactoryBase* FClothingSystemEditorModule::GetFactory()
diff --git a/Engine/Source/Editor/ClothingSystemEditor/Private/SimulationEditorExtenderNv.cpp b/Engine/Source/Editor/ClothingSystemEditor/Private/SimulationEditorExtenderNv.cpp
--- a/Engine/Source/Editor/ClothingSystemEditor/Private/SimulationEditorExtenderNv.cpp
+++ b/Engine/Source/Editor/ClothingSystemEditor/Private/SimulationEditorExtenderNv.cpp
@@ -1,156 +1,3 @@
 // Copyright Epic Games, Inc. All Rights Reserved.
 
 #include "SimulationEditorExtenderNv.h"
-
-#if WITH_NVCLOTH
-
-#include "ClothingSimulation.h"
-#include "ClothingSimulationNv.h"
-#include "ClothingSimulationFactoryNv.h"
-#include "IPersonaPreviewScene.h"
-#include "Animation/DebugSkelMeshComponent.h"
-#include "Framework/Commands/UIAction.h"
-#include "Framework/MultiBox/MultiBoxBuilder.h"
-
-#define LOCTEXT_NAMESPACE "NvSimEditorExtender"
-
-FNvVisualizationOptions::FNvVisualizationOptions()
-{
-	FMemory::Memset(Flags, 0, (int32)EVisualizationOption::Max);
-
-	OptionData[(int32)EVisualizationOption::PhysMesh].DisplayName = LOCTEXT("VisName_PhysMesh", "Physical Mesh");
-	OptionData[(int32)EVisualizationOption::PhysMesh].ToolTip = LOCTEXT("VisName_PhysMesh_ToolTip", "Draws the current physical mesh result");
-
-	OptionData[(int32)EVisualizationOption::Normals].DisplayName = LOCTEXT("VisName_Normals", "Simulation Normals");
-	OptionData[(int32)EVisualizationOption::Normals].ToolTip = LOCTEXT("VisName_Normals_ToolTip", "Draws the current normals for the simulation mesh");
-
-	OptionData[(int32)EVisualizationOption::Collision].DisplayName = LOCTEXT("VisName_Collision", "Collisions");
-	OptionData[(int32)EVisualizationOption::Collision].ToolTip= LOCTEXT("VisName_Collision_ToolTip", "Draws the collision bodies the simulation is currently using");
-
-	OptionData[(int32)EVisualizationOption::Backstop].DisplayName = LOCTEXT("VisName_Backstop", "Backstops");
-	OptionData[(int32)EVisualizationOption::Backstop].ToolTip = LOCTEXT("VisName_Backstop_ToolTip", "Draws the backstop offset for each simulation particle");
-
-	OptionData[(int32)EVisualizationOption::MaxDistances].DisplayName = LOCTEXT("VisName_MaxDistance", "Max Distances");
-	OptionData[(int32)EVisualizationOption::MaxDistances].ToolTip = LOCTEXT("VisName_MaxDistance_ToolTip", "Draws the current max distances for the sim particles as a line along its normal");
-
-	OptionData[(int32)EVisualizationOption::SelfCollision].DisplayName = LOCTEXT("VisName_SelfCollision", "Self Collision Radii");
-	OptionData[(int32)EVisualizationOption::SelfCollision].ToolTip = LOCTEXT("VisName_SelfCollision_ToolTip", "Draws the self collision radius for each particle if self collision is enabled");
-
-	OptionData[(int32)EVisualizationOption::AnimDrive].DisplayName = LOCTEXT("VisName_AnimDrive", "Anim Drive");
-	OptionData[(int32)EVisualizationOption::AnimDrive].ToolTip = LOCTEXT("VisName_AnimDrive_Tooltip", "Draws the current skinned reference mesh for the simulation which anim drive will attempt to reach if enabled");
-
-	OptionData[(int32)EVisualizationOption::Backstop].bDisablesSimulation = true;
-	OptionData[(int32)EVisualizationOption::MaxDistances].bDisablesSimulation = true;
-}
-
-bool FNvVisualizationOptions::ShouldDisableSimulation() const
-{
-	for(int32 OptionIndex = 0; OptionIndex < (int32)FNvVisualizationOptions::EVisualizationOption::Max; ++OptionIndex)
-	{
-		if(Flags[OptionIndex])
-		{
-			const FOptionData& Data = OptionData[OptionIndex];
-
-			if(Data.bDisablesSimulation)
-			{
-				return true;
-			}
-		}
-	}
-
-	return false;
-}
-
-UClass* FSimulationEditorExtenderNv::GetSupportedSimulationFactoryClass()
-{
-	return UClothingSimulationFactoryNv::StaticClass();
-}
-
-void FSimulationEditorExtenderNv::ExtendViewportShowMenu(FMenuBuilder& InMenuBuilder, TSharedRef<IPersonaPreviewScene> InPreviewScene)
-{
-	InMenuBuilder.BeginSection(TEXT("NvSim_Visualizations"), LOCTEXT("VisSection", "Visualizations"));
-	{
-		for(int32 OptionIndex = 0; OptionIndex < (int32)FNvVisualizationOptions::EVisualizationOption::Max; ++OptionIndex)
-		{
-			FNvVisualizationOptions::EVisualizationOption Option = (FNvVisualizationOptions::EVisualizationOption)OptionIndex;
-
-			FExecuteAction ExecuteAction = FExecuteAction::CreateRaw(this, &FSimulationEditorExtenderNv::OnEntryClicked, Option, InPreviewScene);
-			FIsActionChecked IsActionChecked = FIsActionChecked::CreateRaw(this, &FSimulationEditorExtenderNv::IsEntryChecked, Option);
-
-			FUIAction Action(ExecuteAction, FCanExecuteAction(), IsActionChecked);
-
-			InMenuBuilder.AddMenuEntry(VisulizationFlags.OptionData[OptionIndex].DisplayName, VisulizationFlags.OptionData[OptionIndex].ToolTip, FSlateIcon(), Action, NAME_None, EUserInterfaceActionType::ToggleButton);
-		}
-	}
-	InMenuBuilder.EndSection();
-}
-
-void FSimulationEditorExtenderNv::DebugDrawSimulation(const IClothingSimulation* InSimulation, USkeletalMeshComponent* InOwnerComponent, FPrimitiveDrawInterface* PDI)
-{
-	if(!InSimulation)
-	{
-		return;
-	}
-
-	const FClothingSimulationNv* NvSimulation = static_cast<const FClothingSimulationNv*>(InSimulation);
-
-	for(int32 OptionIndex = 0; OptionIndex < (int32)FNvVisualizationOptions::EVisualizationOption::Max; ++OptionIndex)
-	{
-		FNvVisualizationOptions::EVisualizationOption Option = (FNvVisualizationOptions::EVisualizationOption)OptionIndex;
-
-		if(VisulizationFlags.IsSet(Option))
-		{
-			switch(Option)
-			{
-				default:
-					break;
-				case FNvVisualizationOptions::EVisualizationOption::PhysMesh:
-					NvSimulation->DebugDraw_PhysMesh(InOwnerComponent, PDI);
-					break;
-				case FNvVisualizationOptions::EVisualizationOption::Normals:
-					NvSimulation->DebugDraw_Normals(InOwnerComponent, PDI);
-					break;
-				case FNvVisualizationOptions::EVisualizationOption::Collision:
-					NvSimulation->DebugDraw_Collision(InOwnerComponent, PDI);
-					break;
-				case FNvVisualizationOptions::EVisualizationOption::Backstop:
-					NvSimulation->DebugDraw_Backstops(InOwnerComponent, PDI);
-					break;
-				case FNvVisualizationOptions::EVisualizationOption::MaxDistances:
-					NvSimulation->DebugDraw_MaxDistances(InOwnerComponent, PDI);
-					break;
-				case FNvVisualizationOptions::EVisualizationOption::SelfCollision:
-					NvSimulation->DebugDraw_SelfCollision(InOwnerComponent, PDI);
-					break;
-				case FNvVisualizationOptions::EVisualizationOption::AnimDrive:
-					NvSimulation->DebugDraw_AnimDrive(InOwnerComponent, PDI);
-					break;
-			}
-		}
-	}
-}
-
-void FSimulationEditorExtenderNv::OnEntryClicked(FNvVisualizationOptions::EVisualizationOption InOption, TSharedRef<IPersonaPreviewScene> InPreviewScene)
-{
-	VisulizationFlags.Toggle(InOption);
-
-	if(UDebugSkelMeshComponent* MeshComponent = InPreviewScene->GetPreviewMeshComponent())
-	{
-		bool bShouldDisableSim = VisulizationFlags.ShouldDisableSimulation();
-
-		// If we need to toggle the disabled state, handle it
-		if(bShouldDisableSim && MeshComponent->bDisableClothSimulation!= bShouldDisableSim)
-		{
-			MeshComponent->bDisableClothSimulation = !MeshComponent->bDisableClothSimulation;
-		}
-	}
-}
-
-bool FSimulationEditorExtenderNv::IsEntryChecked(FNvVisualizationOptions::EVisualizationOption InOption) const
-{
-	return VisulizationFlags.IsSet(InOption);
-}
-
-#undef LOCTEXT_NAMESPACE
-
-#endif
